AnimationModifierLibrary plugin

diff --git a/Engine/Plugins/Animation/AnimationModifierLibrary/AnimationModifierLibrary.uplugin b/Engine/Plugins/Animation/AnimationModifierLibrary/AnimationModifierLibrary.uplugin
new file mode 100644
--- /dev/null
+++ b/Engine/Plugins/Animation/AnimationModifierLibrary/AnimationModifierLibrary.uplugin
@@ -0,0 +1,26 @@
+{
+	"FileVersion": 1,
+	"Version": 1,
+	"VersionName": "1.0",
+	"FriendlyName": "Animation Modifier Library",
+	"Description": "Collection of Animation Modifiers",
+	"Category": "Animation",
+	"CreatedBy": "Epic Games, Inc.",
+	"CreatedByURL": "http://epicgames.com",
+	"DocsURL": "",
+	"MarketplaceURL": "",
+	"SupportURL": "",
+	"EnabledByDefault": true,
+	"CanContainContent": false,
+	"IsBetaVersion": false,
+	"IsExperimentalVersion": false,
+	"Installed": false,
+	"Modules": 
+	[
+		{
+			"Name": "AnimationModifierLibrary",
+			"Type": "Editor",
+			"LoadingPhase": "Default"
+		}
+	]
+}
diff --git a/Engine/Plugins/Animation/AnimationModifierLibrary/Source/AnimationModifierLibrary/AnimationModifierLibrary.Build.cs b/Engine/Plugins/Animation/AnimationModifierLibrary/Source/AnimationModifierLibrary/AnimationModifierLibrary.Build.cs
new file mode 100644
--- /dev/null
+++ b/Engine/Plugins/Animation/AnimationModifierLibrary/Source/AnimationModifierLibrary/AnimationModifierLibrary.Build.cs
@@ -0,0 +1,36 @@
+// Copyright Epic Games, Inc. All Rights Reserved.
+
+using System.IO;
+
+namespace UnrealBuildTool.Rules
+{
+	public class AnimationModifierLibrary : ModuleRules
+	{
+		public AnimationModifierLibrary(ReadOnlyTargetRules Target) : base(Target)
+		{
+			PublicDependencyModuleNames.AddRange(
+				new string[]
+				{
+				}
+			);
+
+			PrivateDependencyModuleNames.AddRange(
+				new string[]
+				{
+					"Core",
+					"CoreUObject",
+					"Engine",
+					"AnimationModifiers",
+					"AnimationBlueprintLibrary"
+				}
+			);
+
+			PrivateIncludePathModuleNames.AddRange(
+				new string[]
+				{
+				}
+			);
+
+		}
+	}
+}
diff --git a/Engine/Plugins/Animation/AnimationModifierLibrary/Source/AnimationModifierLibrary/Private/AnimationModifierLibrary.cpp b/Engine/Plugins/Animation/AnimationModifierLibrary/Source/AnimationModifierLibrary/Private/AnimationModifierLibrary.cpp
new file mode 100644
--- /dev/null
+++ b/Engine/Plugins/Animation/AnimationModifierLibrary/Source/AnimationModifierLibrary/Private/AnimationModifierLibrary.cpp
@@ -0,0 +1,17 @@
+// Copyright Epic Games, Inc. All Rights Reserved.
+
+#include "AnimationModifierLibrary.h"
+
+#define LOCTEXT_NAMESPACE "AnimationModifierLibraryModule"
+
+void FAnimationModifierLibraryModule::StartupModule()
+{
+}
+
+void FAnimationModifierLibraryModule::ShutdownModule()
+{
+}
+
+#undef LOCTEXT_NAMESPACE
+	
+IMPLEMENT_MODULE(FAnimationModifierLibraryModule, AnimationModifierLibrary)
\ No newline at end of file
diff --git a/Engine/Plugins/Animation/AnimationModifierLibrary/Source/AnimationModifierLibrary/Private/MotionExtractorModifier.cpp b/Engine/Plugins/Animation/AnimationModifierLibrary/Source/AnimationModifierLibrary/Private/MotionExtractorModifier.cpp
new file mode 100644
--- /dev/null
+++ b/Engine/Plugins/Animation/AnimationModifierLibrary/Source/AnimationModifierLibrary/Private/MotionExtractorModifier.cpp
@@ -0,0 +1,289 @@
+// Copyright Epic Games, Inc. All Rights Reserved.
+
+#include "MotionExtractorModifier.h"
+#include "AnimationBlueprintLibrary.h"
+#include "Animation/AnimSequence.h"
+#include "Animation/Skeleton.h"
+#include "BonePose.h"
+#include "Animation/AnimationPoseData.h"
+
+UMotionExtractorModifier::UMotionExtractorModifier()
+	:Super()
+{
+	BoneName = FName(TEXT("root"));
+	MotionType = EMotionExtractor_MotionType::Translation;
+	Axis = EMotionExtractor_Axis::Y;
+	bAbsoluteValue = false;
+	bComponentSpace = true;
+	bUseCustomCurveName = false;
+	CustomCurveName = NAME_None;
+	SampleRate = 30;
+}
+
+void UMotionExtractorModifier::OnApply_Implementation(UAnimSequence* Animation)
+{
+	if (Animation == nullptr)
+	{
+		UE_LOG(LogAnimation, Error, TEXT("MotionExtractorModifier failed. Reason: Invalid Animation"));
+		return;
+	}
+
+	USkeleton* Skeleton = Animation->GetSkeleton();
+	if (Skeleton == nullptr)
+	{
+		UE_LOG(LogAnimation, Error, TEXT("MotionExtractorModifier failed. Reason: Animation with invalid Skeleton. Animation: %s"),
+			*GetNameSafe(Animation));
+		return;
+	}
+
+	const int32 BoneIndex = Skeleton->GetReferenceSkeleton().FindBoneIndex(BoneName);
+	if (BoneIndex == INDEX_NONE)
+	{
+		UE_LOG(LogAnimation, Error, TEXT("MotionExtractorModifier failed. Reason: Invalid Bone Index. BoneName: %s Animation: %s Skeleton: %s"),
+			*BoneName.ToString(), *GetNameSafe(Animation), *GetNameSafe(Skeleton));
+		return;
+	}
+
+	// Ideally we would disable these options when any of those motion types are selected but AnimModifier doesn't support Details Customization atm.
+	if((MotionType == EMotionExtractor_MotionType::Translation || MotionType == EMotionExtractor_MotionType::Rotation || MotionType == EMotionExtractor_MotionType::Scale) && Axis > EMotionExtractor_Axis::Z)
+	{
+		UE_LOG(LogAnimation, Error, TEXT("MotionExtractorModifier failed. Reason: Only X, Y or Z axes are valid options for the selected motion type"));
+		return;
+	}
+
+	FMemMark Mark(FMemStack::Get());
+
+	bool bForceRootLock = Animation->bForceRootLock;
+	Animation->bForceRootLock = false;
+
+	const FName FinalCurveName = GetCurveName();
+	UAnimationBlueprintLibrary::AddCurve(Animation, FinalCurveName, ERawCurveTrackTypes::RCT_Float, false);
+
+	TArray<FBoneIndexType> RequiredBones;
+	RequiredBones.Add(BoneIndex);
+	Skeleton->GetReferenceSkeleton().EnsureParentsExistAndSort(RequiredBones);
+
+	FBoneContainer BoneContainer(RequiredBones, false, *Skeleton);
+	const FCompactPoseBoneIndex CompactPoseBoneIndex = BoneContainer.MakeCompactPoseIndex(FMeshPoseBoneIndex(BoneIndex));
+
+	const float AnimLength = Animation->GetPlayLength();
+	const float SampleInterval = 1.f / SampleRate;
+
+	FTransform LastBoneTransform;
+	float Time = 0.f;
+	int32 SampleIndex = 0;
+	while (Time < AnimLength)
+	{
+		Time = FMath::Clamp(SampleIndex * SampleInterval, 0.f, AnimLength);
+		SampleIndex++;
+
+		const FTransform BoneTransform = ExtractBoneTransform(Animation, BoneContainer, CompactPoseBoneIndex, Time, bComponentSpace);
+
+		// Ignore first frame if we are extracting something that depends on the previous bone transform
+		if (SampleIndex > 1 || (MotionType != EMotionExtractor_MotionType::TranslationSpeed && MotionType != EMotionExtractor_MotionType::RotationSpeed))
+		{
+			const float Value = GetDesiredValue(BoneTransform, LastBoneTransform, SampleInterval);
+
+			UAnimationBlueprintLibrary::AddFloatCurveKey(Animation, FinalCurveName, Time, Value);
+		}
+
+		LastBoneTransform = BoneTransform;
+	}
+
+	Animation->bForceRootLock = bForceRootLock;
+}
+
+void UMotionExtractorModifier::OnRevert_Implementation(UAnimSequence* Animation)
+{
+	// Left empty intentionally. 
+	// Would be nice to have a way to explicitly define if Revert should be called before Apply
+}
+
+FName UMotionExtractorModifier::GetCurveName() const
+{
+	if (bUseCustomCurveName && !CustomCurveName.IsEqual(NAME_None))
+	{
+		return CustomCurveName;
+	}
+
+	FString MotionTypeStr;
+	switch (MotionType)
+	{
+	case EMotionExtractor_MotionType::Translation:		MotionTypeStr = TEXT("translation"); break;
+	case EMotionExtractor_MotionType::Rotation:			MotionTypeStr = TEXT("rotation"); break;
+	case EMotionExtractor_MotionType::Scale:			MotionTypeStr = TEXT("scale"); break;
+	case EMotionExtractor_MotionType::TranslationSpeed:	MotionTypeStr = TEXT("translation_speed"); break;
+	case EMotionExtractor_MotionType::RotationSpeed:	MotionTypeStr = TEXT("rotation_speed"); break;
+	default: check(false); break;
+	}
+
+	FString AxisStr;
+	switch (Axis)
+	{
+	case EMotionExtractor_Axis::X:		AxisStr = TEXT("X"); break;
+	case EMotionExtractor_Axis::Y:		AxisStr = TEXT("Y"); break;
+	case EMotionExtractor_Axis::Z:		AxisStr = TEXT("Z"); break;
+	case EMotionExtractor_Axis::XY:		AxisStr = TEXT("XY"); break;
+	case EMotionExtractor_Axis::XZ:		AxisStr = TEXT("XZ"); break;
+	case EMotionExtractor_Axis::YZ:		AxisStr = TEXT("YZ"); break;
+	case EMotionExtractor_Axis::XYZ:	AxisStr = TEXT("XYZ"); break;
+	default: check(false); break;
+	}
+
+	return FName(*FString::Printf(TEXT("%s_%s_%s"), *BoneName.ToString(), *MotionTypeStr, *AxisStr));
+}
+
+float UMotionExtractorModifier::GetDesiredValue(const FTransform& BoneTransform, const FTransform& LastBoneTransform, float DeltaTime) const
+{
+	float Value = 0.f;
+	switch (MotionType)
+	{
+	case EMotionExtractor_MotionType::Translation:
+	{
+		const FVector Translation = BoneTransform.GetTranslation();
+		switch (Axis)
+		{
+		case EMotionExtractor_Axis::X: Value = Translation.X; break;
+		case EMotionExtractor_Axis::Y: Value = Translation.Y; break;
+		case EMotionExtractor_Axis::Z: Value = Translation.Z; break;
+		default: break;
+		}
+		break;
+	}
+	case EMotionExtractor_MotionType::Rotation:
+	{
+		const FRotator Rotation = BoneTransform.GetRotation().Rotator();
+		switch (Axis)
+		{
+		case EMotionExtractor_Axis::X: Value = Rotation.Roll; break;
+		case EMotionExtractor_Axis::Y: Value = Rotation.Pitch; break;
+		case EMotionExtractor_Axis::Z: Value = Rotation.Yaw; break;
+		default: break;
+		}
+		break;
+	}
+	case EMotionExtractor_MotionType::Scale:
+	{
+		const FVector Scale = BoneTransform.GetScale3D();
+		switch (Axis)
+		{
+		case EMotionExtractor_Axis::X: Value = Scale.X; break;
+		case EMotionExtractor_Axis::Y: Value = Scale.Y; break;
+		case EMotionExtractor_Axis::Z: Value = Scale.Z; break;
+		default: break;
+		}
+		break;
+	}
+	case EMotionExtractor_MotionType::TranslationSpeed:
+	{
+		if (!FMath::IsNearlyZero(DeltaTime))
+		{
+			const float Delta = CalculateMagnitude((BoneTransform.GetTranslation() - LastBoneTransform.GetTranslation()), Axis);
+			Value = Delta / DeltaTime;
+		}
+
+		break;
+	}
+	case EMotionExtractor_MotionType::RotationSpeed:
+	{
+		if(!FMath::IsNearlyZero(DeltaTime))
+		{
+			FQuat Delta = FQuat::Identity;
+			FRotator Rotator = BoneTransform.GetRotation().Rotator();
+			FRotator LastRotator = LastBoneTransform.GetRotation().Rotator();
+
+			switch (Axis)
+			{
+			case EMotionExtractor_Axis::X:		Delta = FQuat(FRotator(0.f, 0.f, Rotator.Roll)) * FQuat(FRotator(0.f, 0.f, LastRotator.Roll)).Inverse(); break;
+			case EMotionExtractor_Axis::Y:		Delta = FQuat(FRotator(Rotator.Pitch, 0.f, 0.f)) * FQuat(FRotator(LastRotator.Pitch, 0.f, 0.f)).Inverse(); break;
+			case EMotionExtractor_Axis::Z:		Delta = FQuat(FRotator(0.f, Rotator.Yaw, 0.f)) * FQuat(FRotator(0.f, LastRotator.Yaw, 0.f)).Inverse(); break;
+			case EMotionExtractor_Axis::XY:		Delta = FQuat(FRotator(Rotator.Pitch, 0.f, Rotator.Roll)) * FQuat(FRotator(LastRotator.Pitch, 0.f, LastRotator.Roll)).Inverse(); break;
+			case EMotionExtractor_Axis::XZ:		Delta = FQuat(FRotator(0.f, Rotator.Yaw, Rotator.Roll)) * FQuat(FRotator(0.f, LastRotator.Yaw, LastRotator.Roll)).Inverse(); break;
+			case EMotionExtractor_Axis::YZ:		Delta = FQuat(FRotator(Rotator.Pitch, Rotator.Yaw, 0.f)) * FQuat(FRotator(LastRotator.Pitch, LastRotator.Yaw, 0.f)).Inverse(); break;
+			case EMotionExtractor_Axis::XYZ:	Delta = BoneTransform.GetRotation() * LastBoneTransform.GetRotation().Inverse(); break;
+			default: break;
+			}
+		
+			float RotationAngle = 0.f;
+			FVector RotationAxis;
+			Delta.ToAxisAndAngle(RotationAxis, RotationAngle);
+			RotationAngle = FMath::UnwindRadians(RotationAngle);
+			if (RotationAngle < 0.0f)
+			{
+				RotationAxis = -RotationAxis;
+				RotationAngle = -RotationAngle;
+			}
+
+			RotationAngle = FMath::RadiansToDegrees(RotationAngle);
+			Value = RotationAngle / DeltaTime;
+		}
+
+		break;
+	}
+	default: check(false); break;
+	}
+
+	if (bAbsoluteValue)
+	{
+		Value = FMath::Abs(Value);
+	}
+
+	if (MathOperation != EMotionExtractor_MathOperation::None)
+	{
+		switch (MathOperation)
+		{
+		case EMotionExtractor_MathOperation::Addition:		Value = Value + Modifier; break;
+		case EMotionExtractor_MathOperation::Subtraction:	Value = Value - Modifier; break;
+		case EMotionExtractor_MathOperation::Division:		Value = Value / Modifier; break;
+		case EMotionExtractor_MathOperation::Multiplication: Value = Value * Modifier; break;
+		default: check(false); break;
+		}
+	}
+
+	return Value;
+}
+
+FTransform UMotionExtractorModifier::ExtractBoneTransform(UAnimSequence* Animation, const FBoneContainer& BoneContainer, FCompactPoseBoneIndex CompactPoseBoneIndex, float Time, bool bComponentSpace)
+{
+	FCompactPose Pose;
+	Pose.SetBoneContainer(&BoneContainer);
+
+	FBlendedCurve Curve;
+	Curve.InitFrom(BoneContainer);
+
+	FAnimExtractContext Context(Time, false);
+	UE::Anim::FStackAttributeContainer Attributes;
+	FAnimationPoseData AnimationPoseData(Pose, Curve, Attributes);
+
+	Animation->GetBonePose(AnimationPoseData, Context, true);
+
+	check(Pose.IsValidIndex(CompactPoseBoneIndex));
+
+	if (bComponentSpace)
+	{
+		FCSPose<FCompactPose> ComponentSpacePose;
+		ComponentSpacePose.InitPose(Pose);
+		return ComponentSpacePose.GetComponentSpaceTransform(CompactPoseBoneIndex);
+	}
+	else
+	{
+		return Pose[CompactPoseBoneIndex];
+	}
+}
+
+float UMotionExtractorModifier::CalculateMagnitude(const FVector& Vector, EMotionExtractor_Axis Axis)
+{
+	switch (Axis)
+	{
+	case EMotionExtractor_Axis::X:		return FMath::Abs(Vector.X); break;
+	case EMotionExtractor_Axis::Y:		return FMath::Abs(Vector.Y); break;
+	case EMotionExtractor_Axis::Z:		return FMath::Abs(Vector.Z); break;
+	case EMotionExtractor_Axis::XY:		return FMath::Sqrt(Vector.X * Vector.X + Vector.Y * Vector.Y); break;
+	case EMotionExtractor_Axis::XZ:		return FMath::Sqrt(Vector.X * Vector.X + Vector.Z * Vector.Z); break;
+	case EMotionExtractor_Axis::YZ:		return FMath::Sqrt(Vector.Y * Vector.Y + Vector.Z * Vector.Z); break;
+	case EMotionExtractor_Axis::XYZ:	return FMath::Sqrt(Vector.X * Vector.X + Vector.Y * Vector.Y + Vector.Z * Vector.Z); break;
+	default: check(false); break;
+	}
+
+	return 0.f;
+}
\ No newline at end of file
diff --git a/Engine/Plugins/Animation/AnimationModifierLibrary/Source/AnimationModifierLibrary/Public/AnimationModifierLibrary.h b/Engine/Plugins/Animation/AnimationModifierLibrary/Source/AnimationModifierLibrary/Public/AnimationModifierLibrary.h
new file mode 100644
--- /dev/null
+++ b/Engine/Plugins/Animation/AnimationModifierLibrary/Source/AnimationModifierLibrary/Public/AnimationModifierLibrary.h
@@ -0,0 +1,15 @@
+// Copyright Epic Games, Inc. All Rights Reserved.
+
+#pragma once
+
+#include "CoreMinimal.h"
+#include "Modules/ModuleManager.h"
+
+class FAnimationModifierLibraryModule : public IModuleInterface
+{
+public:
+
+	/** IModuleInterface implementation */
+	virtual void StartupModule() override;
+	virtual void ShutdownModule() override;
+};
\ No newline at end of file
diff --git a/Engine/Plugins/Animation/AnimationModifierLibrary/Source/AnimationModifierLibrary/Public/MotionExtractorModifier.h b/Engine/Plugins/Animation/AnimationModifierLibrary/Source/AnimationModifierLibrary/Public/MotionExtractorModifier.h
new file mode 100644
--- /dev/null
+++ b/Engine/Plugins/Animation/AnimationModifierLibrary/Source/AnimationModifierLibrary/Public/MotionExtractorModifier.h
@@ -0,0 +1,108 @@
+// Copyright Epic Games, Inc. All Rights Reserved.
+
+#pragma once
+
+#include "CoreMinimal.h"
+#include "AnimationModifier.h"
+#include "MotionExtractorModifier.generated.h"
+
+/** Type of motion to extract */
+UENUM(BlueprintType)
+enum class EMotionExtractor_MotionType : uint8
+{
+	Translation,
+	Rotation,
+	Scale,
+	TranslationSpeed,
+	RotationSpeed
+};
+
+/** Axis to get the final value from */
+UENUM(BlueprintType)
+enum class EMotionExtractor_Axis : uint8
+{
+	X,
+	Y,
+	Z,
+	XY,
+	XZ,
+	YZ,
+	XYZ
+};
+
+/** Math operations that can be applied to the extracted value before add it to the curve */
+UENUM(BlueprintType)
+enum class EMotionExtractor_MathOperation : uint8
+{
+	None,
+	Addition,
+	Subtraction,
+	Division,
+	Multiplication
+};
+
+/** Extracts motion from a bone in the animation and bakes it into a curve */
+UCLASS()
+class UMotionExtractorModifier : public UAnimationModifier
+{
+	GENERATED_BODY()
+
+public:
+
+	/** Bone we are going to generate the curve from */
+	UPROPERTY(EditAnywhere, BlueprintReadWrite, Category = Settings)
+	FName BoneName;
+
+	/** Type of motion to extract */
+	UPROPERTY(EditAnywhere, BlueprintReadWrite, Category = Settings)
+	EMotionExtractor_MotionType MotionType;
+
+	/** Axis to get the value from */
+	UPROPERTY(EditAnywhere, BlueprintReadWrite, Category = Settings)
+	EMotionExtractor_Axis Axis;
+
+	/** Whether to extract the bone pose in component space or local space */
+	UPROPERTY(EditAnywhere, BlueprintReadWrite, Category = Settings)
+	bool bComponentSpace;
+
+	/** Whether to convert the final value to absolute (positive) */
+	UPROPERTY(EditAnywhere, BlueprintReadWrite, Category = Settings)
+	bool bAbsoluteValue;
+
+	/** Optional math operation to apply on the extracted value before add it to the generated curve */
+	UPROPERTY(EditAnywhere, BlueprintReadWrite, Category = Settings)
+	EMotionExtractor_MathOperation MathOperation;
+
+	/** Right operand for the math operation selected */
+	UPROPERTY(EditAnywhere, BlueprintReadWrite, Category = Settings, meta = (EditCondition = "MathOperation != EMotionExtractor_MathOperation::None"))
+	float Modifier;
+
+	/** Rate used to sample the animation */
+	UPROPERTY(EditAnywhere, Category = Settings, meta = (ClampMin = "1"))
+	int32 SampleRate;
+
+	/** Whether we want to specify a custom name for the curve. If false, the name of the curve will be auto generated based on the data we are going to extract */
+	UPROPERTY(EditAnywhere, BlueprintReadWrite, Category = Settings)
+	bool bUseCustomCurveName;
+
+	/** Custom name for the curve we are going to generate. */
+	UPROPERTY(EditAnywhere, BlueprintReadWrite, Category = Settings, meta = (EditCondition = "bUseCustomCurveName"))
+	FName CustomCurveName;
+
+	UMotionExtractorModifier();
+
+	virtual void OnApply_Implementation(UAnimSequence* Animation) override;
+	virtual void OnRevert_Implementation(UAnimSequence* Animation) override;
+
+	/** Returns the name for the curve. If CurveName is None it generates a name based on the data we are going to extract */
+	FName GetCurveName() const;
+
+	/** Returns the desired value from the extracted poses */
+	float GetDesiredValue(const FTransform& BoneTransform, const FTransform& LastBoneTransform, float DeltaTime) const;
+
+	/** Helper function to extract the pose for a given bone at a given time */
+	static FTransform ExtractBoneTransform(UAnimSequence* Animation, const FBoneContainer& BoneContainer, FCompactPoseBoneIndex CompactPoseBoneIndex, float Time, bool bComponentSpace);
+
+	/** Helper function to calculate the magnitude of a vector only considering a specific axis or axes */
+	static float CalculateMagnitude(const FVector& Vector, EMotionExtractor_Axis Axis);
+};
