Added option to remove curves on revert for the MotionExtractorModifier
#rnx

diff --git a/Engine/Plugins/Animation/AnimationModifierLibrary/Source/AnimationModifierLibrary/Private/MotionExtractorModifier.cpp b/Engine/Plugins/Animation/AnimationModifierLibrary/Source/AnimationModifierLibrary/Private/MotionExtractorModifier.cpp
--- a/Engine/Plugins/Animation/AnimationModifierLibrary/Source/AnimationModifierLibrary/Private/MotionExtractorModifier.cpp
+++ b/Engine/Plugins/Animation/AnimationModifierLibrary/Source/AnimationModifierLibrary/Private/MotionExtractorModifier.cpp
@@ -18,6 +18,7 @@ UMotionExtractorModifier::UMotionExtractorModifier()
 	bUseCustomCurveName = false;
 	CustomCurveName = NAME_None;
 	SampleRate = 30;
+	bRemoveCurveOnRevert = false;
 }
 
 void UMotionExtractorModifier::OnApply_Implementation(UAnimSequence* Animation)
@@ -95,8 +96,12 @@ void UMotionExtractorModifier::OnApply_Implementation(UAnimSequence* Animation)
 
 void UMotionExtractorModifier::OnRevert_Implementation(UAnimSequence* Animation)
 {
-	// Left empty intentionally. 
-	// Would be nice to have a way to explicitly define if Revert should be called before Apply
+	if (bRemoveCurveOnRevert)
+	{
+		const FName FinalCurveName = GetCurveName();
+		const bool bRemoveNameFromSkeleton = false;
+		UAnimationBlueprintLibrary::RemoveCurve(Animation, FinalCurveName, bRemoveNameFromSkeleton);
+	}
 }
 
 FName UMotionExtractorModifier::GetCurveName() const
diff --git a/Engine/Plugins/Animation/AnimationModifierLibrary/Source/AnimationModifierLibrary/Public/MotionExtractorModifier.h b/Engine/Plugins/Animation/AnimationModifierLibrary/Source/AnimationModifierLibrary/Public/MotionExtractorModifier.h
--- a/Engine/Plugins/Animation/AnimationModifierLibrary/Source/AnimationModifierLibrary/Public/MotionExtractorModifier.h
+++ b/Engine/Plugins/Animation/AnimationModifierLibrary/Source/AnimationModifierLibrary/Public/MotionExtractorModifier.h
@@ -61,6 +61,12 @@ public:
 	UPROPERTY(EditAnywhere, BlueprintReadWrite, Category = Settings)
 	EMotionExtractor_Axis Axis;
 
+	/** Whether we want to remove the curve when we revert or re-apply modifier 
+		Disabling this allows you to modify settings and create a new curve each time you re-apply the modifier
+		Enabling this is the preferred setting when using a modifier that's applied in bulk and you may want to remove/rename curves */
+	UPROPERTY(EditAnywhere, BlueprintReadWrite, Category = Settings)
+	bool bRemoveCurveOnRevert;
+
 	/** Whether to extract the bone pose in component space or local space */
 	UPROPERTY(EditAnywhere, BlueprintReadWrite, Category = Settings)
 	bool bComponentSpace;
