AnimationModifierLibrary: Added animation modifier to reorient the root bone in the animation while maintaining mesh position and rotation

diff --git a/Engine/Plugins/Animation/AnimationModifierLibrary/Source/AnimationModifierLibrary/Private/ReOrientRootBoneModifier.cpp b/Engine/Plugins/Animation/AnimationModifierLibrary/Source/AnimationModifierLibrary/Private/ReOrientRootBoneModifier.cpp
new file mode 100644
--- /dev/null
+++ b/Engine/Plugins/Animation/AnimationModifierLibrary/Source/AnimationModifierLibrary/Private/ReOrientRootBoneModifier.cpp
@@ -0,0 +1,101 @@
+// Copyright Epic Games, Inc. All Rights Reserved.
+
+#include "ReOrientRootBoneModifier.h"
+#include "AnimationBlueprintLibrary.h"
+#include "Animation/AnimSequence.h"
+#include "AnimationUtils.h"
+
+#define LOCTEXT_NAMESPACE "ReOrientRootBoneModifier"
+
+UReOrientRootBoneModifier::UReOrientRootBoneModifier()
+	:Super()
+{
+}
+
+void UReOrientRootBoneModifier::OnApply_Implementation(UAnimSequence* Animation)
+{
+	ReOrientRootBone_Internal(Animation, Rotator.Quaternion());
+}
+
+void UReOrientRootBoneModifier::OnRevert_Implementation(UAnimSequence* Animation)
+{
+	ReOrientRootBone_Internal(Animation, Rotator.Quaternion().Inverse());
+}
+
+void UReOrientRootBoneModifier::ReOrientRootBone_Internal(UAnimSequence* Animation, const FQuat& Quat)
+{
+	if (Animation == nullptr)
+	{
+		UE_LOG(LogAnimation, Error, TEXT("ReOrientMeshModifier failed. Reason: Invalid Animation"));
+		return;
+	}
+
+	IAnimationDataController& Controller = Animation->GetController();
+	const UAnimDataModel* Model = Animation->GetDataModel();
+
+	if (Model == nullptr)
+	{
+		UE_LOG(LogAnimation, Error, TEXT("ReOrientMeshModifier failed. Reason: Invalid Data Model. Animation: %s"), *GetNameSafe(Animation));
+		return;
+	}
+
+	const USkeleton* Skeleton = Animation->GetSkeleton();
+	if (Skeleton == nullptr)
+	{
+		UE_LOG(LogAnimation, Error, TEXT("ReOrientMeshModifier failed. Reason: Invalid Skeleton. Animation: %s"), *GetNameSafe(Animation));
+		return;
+	}
+
+	const FReferenceSkeleton& RefSkeleton = Skeleton->GetReferenceSkeleton();
+	if (RefSkeleton.GetNum() == 0)
+	{
+		UE_LOG(LogAnimation, Error, TEXT("ReOrientMeshModifier failed. Reason: Ref Skeleton. Animation: %s"), *GetNameSafe(Animation));
+		return;
+	}
+
+	// Collect tracks for the bones we are going to modify (root and direct children of the root)
+	TArray<const FBoneAnimationTrack*> BonesTracks;
+	BonesTracks.Add(&Model->GetBoneTrackByIndex(0));
+	for (int32 BoneIdx = 1; BoneIdx < RefSkeleton.GetNum(); BoneIdx++)
+	{
+		if (RefSkeleton.GetParentIndex(BoneIdx) == 0)
+		{
+			BonesTracks.Add(&Model->GetBoneTrackByIndex(BoneIdx));
+		}
+	}
+	
+	// Start editing animation data
+	const bool bShouldTransact = false;
+	Controller.OpenBracket(LOCTEXT("ReorientingRootBone_Bracket", "Reorienting root bone"), bShouldTransact);
+	
+	// For each key in the animation
+	const int32 Num = Model->GetNumberOfKeys();
+	for (int32 AnimKey = 0; AnimKey < Num; AnimKey++)
+	{
+		const FInt32Range KeyRangeToSet(AnimKey, AnimKey + 1);
+
+		// Reorient the root bone
+		FTransform RootTransformOriginal;
+		FAnimationUtils::ExtractTransformForFrameFromTrack(BonesTracks[0]->InternalTrackData, AnimKey, RootTransformOriginal);
+
+		FTransform RootTransformNew = RootTransformOriginal;
+		RootTransformNew.SetRotation(RootTransformOriginal.GetRotation() * Quat);
+		Controller.UpdateBoneTrackKeys(BonesTracks[0]->Name, KeyRangeToSet, { RootTransformNew.GetLocation() }, { RootTransformNew.GetRotation() }, { RootTransformNew.GetScale3D() });
+
+		// Now the mesh is facing the wrong axis. Update direct children of the root with the local space transform that puts them back to where they were originally
+		for (int32 Idx = 1; Idx < BonesTracks.Num(); Idx++)
+		{
+			FTransform OtherBoneTransform;
+			FAnimationUtils::ExtractTransformForFrameFromTrack(BonesTracks[Idx]->InternalTrackData, AnimKey, OtherBoneTransform);
+			
+			OtherBoneTransform = OtherBoneTransform * RootTransformOriginal;
+			OtherBoneTransform = OtherBoneTransform.GetRelativeTransform(RootTransformNew);
+			Controller.UpdateBoneTrackKeys(BonesTracks[Idx]->Name, KeyRangeToSet, { OtherBoneTransform.GetLocation() }, { OtherBoneTransform.GetRotation() }, { OtherBoneTransform.GetScale3D() });
+		}
+	}
+
+	// Done editing animation data
+	Controller.CloseBracket(bShouldTransact);
+}
+
+#undef LOCTEXT_NAMESPACE
\ No newline at end of file
diff --git a/Engine/Plugins/Animation/AnimationModifierLibrary/Source/AnimationModifierLibrary/Public/ReOrientRootBoneModifier.h b/Engine/Plugins/Animation/AnimationModifierLibrary/Source/AnimationModifierLibrary/Public/ReOrientRootBoneModifier.h
new file mode 100644
--- /dev/null
+++ b/Engine/Plugins/Animation/AnimationModifierLibrary/Source/AnimationModifierLibrary/Public/ReOrientRootBoneModifier.h
@@ -0,0 +1,30 @@
+// Copyright Epic Games, Inc. All Rights Reserved.
+
+#pragma once
+
+#include "CoreMinimal.h"
+#include "AnimationModifier.h"
+#include "ReOrientRootBoneModifier.generated.h"
+
+/** Reorient root bone in the animation while maintaining mesh position and rotation */
+UCLASS()
+class UReOrientRootBoneModifier : public UAnimationModifier
+{
+	GENERATED_BODY()
+
+public:
+
+	/** Rotation to apply to the root */
+	UPROPERTY(EditAnywhere, BlueprintReadOnly, Category = "Settings")
+	FRotator Rotator;
+
+	UReOrientRootBoneModifier();
+
+	virtual void OnApply_Implementation(UAnimSequence* Animation) override;
+	virtual void OnRevert_Implementation(UAnimSequence* Animation) override;
+
+private:
+
+	void ReOrientRootBone_Internal(UAnimSequence* Animation, const FQuat& Quat);
+
+};
