Fixed bug where the modifier could end up modifying the wrong bone when applied on an animation using a compatible skeleton with a different hierarchy than the original.

diff --git a/Engine/Plugins/Animation/AnimationModifierLibrary/Source/AnimationModifierLibrary/Private/ReOrientRootBoneModifier.cpp b/Engine/Plugins/Animation/AnimationModifierLibrary/Source/AnimationModifierLibrary/Private/ReOrientRootBoneModifier.cpp
--- a/Engine/Plugins/Animation/AnimationModifierLibrary/Source/AnimationModifierLibrary/Private/ReOrientRootBoneModifier.cpp
+++ b/Engine/Plugins/Animation/AnimationModifierLibrary/Source/AnimationModifierLibrary/Private/ReOrientRootBoneModifier.cpp
@@ -60,7 +60,11 @@ void UReOrientRootBoneModifier::ReOrientRootBone_Internal(UAnimSequence* Animati
 	{
 		if (RefSkeleton.GetParentIndex(BoneIdx) == 0)
 		{
-			BonesTracks.Add(&Model->GetBoneTrackByIndex(BoneIdx));
+			const int32 BoneTrackIdx = Model->GetBoneTrackIndexByName(RefSkeleton.GetBoneName(BoneIdx));
+			if(BoneTrackIdx != INDEX_NONE)
+			{
+				BonesTracks.Add(&Model->GetBoneTrackByIndex(BoneTrackIdx));
+			}
 		}
 	}
 	
