AnimModifier to copy transforms from one bone to another

diff --git a/Engine/Plugins/Animation/AnimationModifierLibrary/Source/AnimationModifierLibrary/Private/CopyBonesModifier.cpp b/Engine/Plugins/Animation/AnimationModifierLibrary/Source/AnimationModifierLibrary/Private/CopyBonesModifier.cpp
new file mode 100644
--- /dev/null
+++ b/Engine/Plugins/Animation/AnimationModifierLibrary/Source/AnimationModifierLibrary/Private/CopyBonesModifier.cpp
@@ -0,0 +1,95 @@
+// Copyright Epic Games, Inc. All Rights Reserved.
+
+#include "CopyBonesModifier.h"
+#include "AnimationBlueprintLibrary.h"
+#include "Animation/AnimSequence.h"
+
+#define LOCTEXT_NAMESPACE "CopyBonesModifier"
+
+UCopyBonesModifier::UCopyBonesModifier()
+	:Super()
+{
+}
+
+void UCopyBonesModifier::OnApply_Implementation(UAnimSequence* Animation)
+{
+	if (Animation == nullptr)
+	{
+		UE_LOG(LogAnimation, Error, TEXT("CopyBonesModifier failed. Reason: Invalid Animation"));
+		return;
+	}
+
+	IAnimationDataController& Controller = Animation->GetController();
+	const UAnimDataModel* Model = Animation->GetDataModel();
+
+	if (Model == nullptr)
+	{
+		UE_LOG(LogAnimation, Error, TEXT("CopyBonesModifier failed. Reason: Invalid Data Model. Animation: %s"), *GetNameSafe(Animation));
+		return;
+	}
+
+	// Helper structure to store the extracted keys from source bones
+	struct FCopyBoneData
+	{
+		FName SourceBoneName = NAME_None;
+		FName TargetBoneName = NAME_None;
+		TArray<FVector> PositionalKeys;
+		TArray<FQuat> RotationalKeys;
+		TArray<FVector> ScalingKeys;
+		FCopyBoneData(const FName& InSourceBoneName, const FName& InTargetBoneName, int32 InNumKeys)
+			: SourceBoneName(InSourceBoneName), TargetBoneName(InTargetBoneName)
+		{
+			PositionalKeys.Reserve(InNumKeys);
+			RotationalKeys.Reserve(InNumKeys);
+			ScalingKeys.Reserve(InNumKeys);
+		}
+	};
+
+	const int32 NumKeys = Model->GetNumberOfKeys();
+
+	TArray<FCopyBoneData> CopyBoneDataContainer;
+	CopyBoneDataContainer.Reserve(BonePairs.Num());
+	for (const FBoneReferencePair& Pair : BonePairs)
+	{
+		CopyBoneDataContainer.Add(FCopyBoneData(Pair.SourceBone.BoneName, Pair.TargetBone.BoneName, NumKeys));
+	}
+
+	// Temporally set ForceRootLock to true so we get the correct transforms regardless of the root motion configuration in the animation
+	TGuardValue<bool> ForceRootLockGuard(Animation->bForceRootLock, true);
+
+	// Get the transform of all the source bones in the desired space
+	for (int32 AnimKey = 0; AnimKey < NumKeys; AnimKey++)
+	{
+		for (FCopyBoneData& Data : CopyBoneDataContainer)
+		{
+			FAnimPose AnimPose;
+			UAnimPoseExtensions::GetAnimPoseAtFrame(Animation, AnimKey, FAnimPoseEvaluationOptions(), AnimPose);
+
+			FTransform BonePose = UAnimPoseExtensions::GetBonePose(AnimPose, Data.SourceBoneName, BonePoseSpace);
+			Data.PositionalKeys.Add(BonePose.GetLocation());
+			Data.RotationalKeys.Add(BonePose.GetRotation());
+			Data.ScalingKeys.Add(BonePose.GetScale3D());
+		}
+	}
+
+	// Start editing animation data
+	const bool bShouldTransact = false;
+	Controller.OpenBracket(LOCTEXT("CopyBonesModifier_Bracket", "Updating bones"), bShouldTransact);
+
+	// Copy all the transforms from source bones to target bones
+	for (const FCopyBoneData& Data : CopyBoneDataContainer)
+	{
+		const FInt32Range KeyRangeToSet(0, NumKeys);
+		Controller.UpdateBoneTrackKeys(Data.TargetBoneName, KeyRangeToSet, Data.PositionalKeys, Data.RotationalKeys, Data.ScalingKeys);
+	}
+
+	// Done editing animation data
+	Controller.CloseBracket(bShouldTransact);
+}
+
+void UCopyBonesModifier::OnRevert_Implementation(UAnimSequence* Animation)
+{
+	// This AnimModifier doesn't support Revert operation
+}
+
+#undef LOCTEXT_NAMESPACE
\ No newline at end of file
diff --git a/Engine/Plugins/Animation/AnimationModifierLibrary/Source/AnimationModifierLibrary/Public/CopyBonesModifier.h b/Engine/Plugins/Animation/AnimationModifierLibrary/Source/AnimationModifierLibrary/Public/CopyBonesModifier.h
new file mode 100644
--- /dev/null
+++ b/Engine/Plugins/Animation/AnimationModifierLibrary/Source/AnimationModifierLibrary/Public/CopyBonesModifier.h
@@ -0,0 +1,45 @@
+// Copyright Epic Games, Inc. All Rights Reserved.
+
+#pragma once
+
+#include "CoreMinimal.h"
+#include "AnimationModifier.h"
+#include "AnimPose.h"
+#include "CopyBonesModifier.generated.h"
+
+USTRUCT(BlueprintType)
+struct FBoneReferencePair
+{
+	GENERATED_BODY()
+
+	/** Bone to get transform from  */
+	UPROPERTY(EditAnywhere, Category = "Default")
+	FBoneReference SourceBone;
+
+	/** Bone to update with the transform copied from SourceBone */
+	UPROPERTY(EditAnywhere, Category = "Default")
+	FBoneReference TargetBone;
+};
+
+/** Animation Modifier to copy the transform of 'SourceBone(s)' to 'TargetBone(s)' */
+UCLASS()
+class UCopyBonesModifier : public UAnimationModifier
+{
+	GENERATED_BODY()
+
+public:
+
+	/** Source and Target bone pairs */
+	UPROPERTY(EditAnywhere, BlueprintReadOnly, Category = "Settings")
+	TArray<FBoneReferencePair> BonePairs;
+
+	/** Space to convert SourceBone transforms into before copying them to TargetBone */
+	UPROPERTY(EditAnywhere, BlueprintReadOnly, Category = "Settings")
+	EAnimPoseSpaces BonePoseSpace = EAnimPoseSpaces::World;
+
+	UCopyBonesModifier();
+
+	virtual void OnApply_Implementation(UAnimSequence* Animation) override;
+	virtual void OnRevert_Implementation(UAnimSequence* Animation) override;
+
+};
