MotionExtractorModifier: Added option to get normalized values and option to get values relative to the first frame in the animation.

diff --git a/Engine/Plugins/Animation/AnimationModifierLibrary/Source/AnimationModifierLibrary/Private/MotionExtractorModifier.cpp b/Engine/Plugins/Animation/AnimationModifierLibrary/Source/AnimationModifierLibrary/Private/MotionExtractorModifier.cpp
--- a/Engine/Plugins/Animation/AnimationModifierLibrary/Source/AnimationModifierLibrary/Private/MotionExtractorModifier.cpp
+++ b/Engine/Plugins/Animation/AnimationModifierLibrary/Source/AnimationModifierLibrary/Private/MotionExtractorModifier.cpp
@@ -14,7 +14,9 @@ UMotionExtractorModifier::UMotionExtractorModifier()
 	MotionType = EMotionExtractor_MotionType::Translation;
 	Axis = EMotionExtractor_Axis::Y;
 	bAbsoluteValue = false;
+	bRelativeToFirstFrame = false;
 	bComponentSpace = true;
+	bNormalize = false;
 	bUseCustomCurveName = false;
 	CustomCurveName = NAME_None;
 	SampleRate = 30;
@@ -54,8 +56,7 @@ void UMotionExtractorModifier::OnApply_Implementation(UAnimSequence* Animation)
 
 	FMemMark Mark(FMemStack::Get());
 
-	bool bForceRootLock = Animation->bForceRootLock;
-	Animation->bForceRootLock = false;
+	TGuardValue<bool> ForceRootLockGuard(Animation->bForceRootLock, false);
 
 	const FName FinalCurveName = GetCurveName();
 	UAnimationBlueprintLibrary::AddCurve(Animation, FinalCurveName, ERawCurveTrackTypes::RCT_Float, false);
@@ -67,31 +68,67 @@ void UMotionExtractorModifier::OnApply_Implementation(UAnimSequence* Animation)
 	FBoneContainer BoneContainer(RequiredBones, false, *Skeleton);
 	const FCompactPoseBoneIndex CompactPoseBoneIndex = BoneContainer.MakeCompactPoseIndex(FMeshPoseBoneIndex(BoneIndex));
 
+	FTransform FirstFrameBoneTransform = ExtractBoneTransform(Animation, BoneContainer, CompactPoseBoneIndex, 0.f, bComponentSpace);
+
 	const float AnimLength = Animation->GetPlayLength();
 	const float SampleInterval = 1.f / SampleRate;
 
-	FTransform LastBoneTransform;
+	FTransform LastBoneTransform = FTransform::Identity;
 	float Time = 0.f;
 	int32 SampleIndex = 0;
+
+	float MaxValue = -MAX_FLT;
+	if(bNormalize)
+	{
+		while (Time < AnimLength)
+		{
+			Time = FMath::Clamp(SampleIndex * SampleInterval, 0.f, AnimLength);
+			SampleIndex++;
+
+			FTransform BoneTransform = ExtractBoneTransform(Animation, BoneContainer, CompactPoseBoneIndex, Time, bComponentSpace);
+
+			if(bRelativeToFirstFrame)
+			{
+				BoneTransform = BoneTransform.GetRelativeTransform(FirstFrameBoneTransform);
+			}
+
+			// Ignore first frame if we are extracting something that depends on the previous bone transform
+			if (SampleIndex > 1 || (MotionType != EMotionExtractor_MotionType::TranslationSpeed && MotionType != EMotionExtractor_MotionType::RotationSpeed))
+			{
+				const float Value = GetDesiredValue(BoneTransform, LastBoneTransform, SampleInterval);
+				MaxValue = FMath::Max(FMath::Abs(Value), MaxValue);
+			}
+
+			LastBoneTransform = BoneTransform;
+		}
+	}
+	
+	LastBoneTransform = FTransform::Identity;
+	Time = 0.f;
+	SampleIndex = 0;
 	while (Time < AnimLength)
 	{
 		Time = FMath::Clamp(SampleIndex * SampleInterval, 0.f, AnimLength);
 		SampleIndex++;
 
-		const FTransform BoneTransform = ExtractBoneTransform(Animation, BoneContainer, CompactPoseBoneIndex, Time, bComponentSpace);
+		FTransform BoneTransform = ExtractBoneTransform(Animation, BoneContainer, CompactPoseBoneIndex, Time, bComponentSpace);
+
+		if (bRelativeToFirstFrame)
+		{
+			BoneTransform = BoneTransform.GetRelativeTransform(FirstFrameBoneTransform);
+		}
 
 		// Ignore first frame if we are extracting something that depends on the previous bone transform
 		if (SampleIndex > 1 || (MotionType != EMotionExtractor_MotionType::TranslationSpeed && MotionType != EMotionExtractor_MotionType::RotationSpeed))
 		{
 			const float Value = GetDesiredValue(BoneTransform, LastBoneTransform, SampleInterval);
+			const float FinalValue = (bNormalize && MaxValue != 0.f) ? FMath::Abs(Value) / MaxValue : Value;
 
-			UAnimationBlueprintLibrary::AddFloatCurveKey(Animation, FinalCurveName, Time, Value);
+			UAnimationBlueprintLibrary::AddFloatCurveKey(Animation, FinalCurveName, Time, FinalValue);
 		}
 
 		LastBoneTransform = BoneTransform;
 	}
-
-	Animation->bForceRootLock = bForceRootLock;
 }
 
 void UMotionExtractorModifier::OnRevert_Implementation(UAnimSequence* Animation)
diff --git a/Engine/Plugins/Animation/AnimationModifierLibrary/Source/AnimationModifierLibrary/Public/MotionExtractorModifier.h b/Engine/Plugins/Animation/AnimationModifierLibrary/Source/AnimationModifierLibrary/Public/MotionExtractorModifier.h
--- a/Engine/Plugins/Animation/AnimationModifierLibrary/Source/AnimationModifierLibrary/Public/MotionExtractorModifier.h
+++ b/Engine/Plugins/Animation/AnimationModifierLibrary/Source/AnimationModifierLibrary/Public/MotionExtractorModifier.h
@@ -67,6 +67,10 @@ public:
 	UPROPERTY(EditAnywhere, BlueprintReadWrite, Category = Settings)
 	bool bRemoveCurveOnRevert;
 
+	/** Whether to extract the bone transforms relative to the first frame in the animation */
+	UPROPERTY(EditAnywhere, BlueprintReadWrite, Category = Settings)
+	bool bRelativeToFirstFrame;
+
 	/** Whether to extract the bone pose in component space or local space */
 	UPROPERTY(EditAnywhere, BlueprintReadWrite, Category = Settings)
 	bool bComponentSpace;
@@ -83,6 +87,10 @@ public:
 	UPROPERTY(EditAnywhere, BlueprintReadWrite, Category = Settings, meta = (EditCondition = "MathOperation != EMotionExtractor_MathOperation::None"))
 	float Modifier;
 
+	/** Whether we want a normalized value (0 - 1) */
+	UPROPERTY(EditAnywhere, BlueprintReadWrite, Category = Settings)
+	bool bNormalize;
+
 	/** Rate used to sample the animation */
 	UPROPERTY(EditAnywhere, Category = Settings, meta = (ClampMin = "1"))
 	int32 SampleRate;
