Fix for CopyBonesModifier moving bones to incorrect locations when using world space transforms
#9567

diff --git a/Engine/Plugins/Animation/AnimationModifierLibrary/Source/AnimationModifierLibrary/Private/CopyBonesModifier.cpp b/Engine/Plugins/Animation/AnimationModifierLibrary/Source/AnimationModifierLibrary/Private/CopyBonesModifier.cpp
--- a/Engine/Plugins/Animation/AnimationModifierLibrary/Source/AnimationModifierLibrary/Private/CopyBonesModifier.cpp
+++ b/Engine/Plugins/Animation/AnimationModifierLibrary/Source/AnimationModifierLibrary/Private/CopyBonesModifier.cpp
@@ -66,9 +66,13 @@ void UCopyBonesModifier::OnApply_Implementation(UAnimSequence* Animation)
 			UAnimPoseExtensions::GetAnimPoseAtFrame(Animation, AnimKey, FAnimPoseEvaluationOptions(), AnimPose);
 
 			FTransform BonePose = UAnimPoseExtensions::GetBonePose(AnimPose, Data.SourceBoneName, BonePoseSpace);
-			Data.PositionalKeys.Add(BonePose.GetLocation());
-			Data.RotationalKeys.Add(BonePose.GetRotation());
-			Data.ScalingKeys.Add(BonePose.GetScale3D());
+
+			// UAnimDataController::UpdateBoneTrackKeys expects local transforms so we need to convert the source transforms to target bone local transforms first. 
+			UAnimPoseExtensions::SetBonePose(AnimPose, BonePose, Data.TargetBoneName, BonePoseSpace);
+			FTransform BonePoseTargetLocal = UAnimPoseExtensions::GetBonePose(AnimPose, Data.TargetBoneName, EAnimPoseSpaces::Local);
+			Data.PositionalKeys.Add(BonePoseTargetLocal.GetLocation());
+			Data.RotationalKeys.Add(BonePoseTargetLocal.GetRotation());
+			Data.ScalingKeys.Add(BonePoseTargetLocal.GetScale3D());
 		}
 	}
 
