CopyBoneAnimModifier: Validate intpus and fixed bug where a wrong bone transform was being set when modifying bones in the same chain

diff --git a/Engine/Plugins/Animation/AnimationModifierLibrary/Source/AnimationModifierLibrary/Private/CopyBonesModifier.cpp b/Engine/Plugins/Animation/AnimationModifierLibrary/Source/AnimationModifierLibrary/Private/CopyBonesModifier.cpp
--- a/Engine/Plugins/Animation/AnimationModifierLibrary/Source/AnimationModifierLibrary/Private/CopyBonesModifier.cpp
+++ b/Engine/Plugins/Animation/AnimationModifierLibrary/Source/AnimationModifierLibrary/Private/CopyBonesModifier.cpp
@@ -28,36 +28,49 @@ void UCopyBonesModifier::OnApply_Implementation(UAnimSequence* Animation)
 		return;
 	}
 
-	// Helper structure to store the extracted keys from source bones
+	// Helper structure to store data for the bones we are going to modify
 	struct FCopyBoneData
 	{
 		FName SourceBoneName = NAME_None;
 		FName TargetBoneName = NAME_None;
-		TArray<FVector> PositionalKeys;
-		TArray<FQuat> RotationalKeys;
-		TArray<FVector> ScalingKeys;
-		FCopyBoneData(const FName& InSourceBoneName, const FName& InTargetBoneName, int32 InNumKeys)
-			: SourceBoneName(InSourceBoneName), TargetBoneName(InTargetBoneName)
-		{
-			PositionalKeys.Reserve(InNumKeys);
-			RotationalKeys.Reserve(InNumKeys);
-			ScalingKeys.Reserve(InNumKeys);
-		}
+		int32 SourceBoneIdx = INDEX_NONE;
+		int32 TargetBoneIdx = INDEX_NONE;
+		FCopyBoneData(const FName& InSourceBoneName, const FName& InTargetBoneName, int32 InSourceBoneIdx, int32 InTargetBoneIdx)
+			: SourceBoneName(InSourceBoneName), TargetBoneName(InTargetBoneName), SourceBoneIdx(InSourceBoneIdx), TargetBoneIdx(InTargetBoneIdx) {}
 	};
 
-	const int32 NumKeys = Model->GetNumberOfKeys();
-
+	// Validate input
 	TArray<FCopyBoneData> CopyBoneDataContainer;
 	CopyBoneDataContainer.Reserve(BonePairs.Num());
 	for (const FBoneReferencePair& Pair : BonePairs)
 	{
-		CopyBoneDataContainer.Add(FCopyBoneData(Pair.SourceBone.BoneName, Pair.TargetBone.BoneName, NumKeys));
+		const int32 SourceBoneIdx = Model->GetBoneTrackIndexByName(Pair.SourceBone.BoneName);
+		if (SourceBoneIdx == INDEX_NONE)
+		{
+			continue;
+		}
+
+		const int32 TargetBoneIdx = Model->GetBoneTrackIndexByName(Pair.TargetBone.BoneName);
+		if (TargetBoneIdx == INDEX_NONE)
+		{
+			continue;
+		}
+
+		CopyBoneDataContainer.Add(FCopyBoneData(Pair.SourceBone.BoneName, Pair.TargetBone.BoneName, SourceBoneIdx, TargetBoneIdx));
 	}
+	
+	// Sort bones to modify so we always modify parents first
+	CopyBoneDataContainer.Sort([](const FCopyBoneData& A, const FCopyBoneData& B) { return A.TargetBoneIdx < B.TargetBoneIdx; });
 
 	// Temporally set ForceRootLock to true so we get the correct transforms regardless of the root motion configuration in the animation
 	TGuardValue<bool> ForceRootLockGuard(Animation->bForceRootLock, true);
 
+	// Start editing animation data
+	const bool bShouldTransact = false;
+	Controller.OpenBracket(LOCTEXT("CopyBonesModifier_Bracket", "Updating bones"), bShouldTransact);
+
 	// Get the transform of all the source bones in the desired space
+	const int32 NumKeys = Model->GetNumberOfKeys();
 	for (int32 AnimKey = 0; AnimKey < NumKeys; AnimKey++)
 	{
 		for (FCopyBoneData& Data : CopyBoneDataContainer)
@@ -70,21 +83,10 @@ void UCopyBonesModifier::OnApply_Implementation(UAnimSequence* Animation)
 			// UAnimDataController::UpdateBoneTrackKeys expects local transforms so we need to convert the source transforms to target bone local transforms first. 
 			UAnimPoseExtensions::SetBonePose(AnimPose, BonePose, Data.TargetBoneName, BonePoseSpace);
 			FTransform BonePoseTargetLocal = UAnimPoseExtensions::GetBonePose(AnimPose, Data.TargetBoneName, EAnimPoseSpaces::Local);
-			Data.PositionalKeys.Add(BonePoseTargetLocal.GetLocation());
-			Data.RotationalKeys.Add(BonePoseTargetLocal.GetRotation());
-			Data.ScalingKeys.Add(BonePoseTargetLocal.GetScale3D());
-		}
-	}
 
-	// Start editing animation data
-	const bool bShouldTransact = false;
-	Controller.OpenBracket(LOCTEXT("CopyBonesModifier_Bracket", "Updating bones"), bShouldTransact);
-
-	// Copy all the transforms from source bones to target bones
-	for (const FCopyBoneData& Data : CopyBoneDataContainer)
-	{
-		const FInt32Range KeyRangeToSet(0, NumKeys);
-		Controller.UpdateBoneTrackKeys(Data.TargetBoneName, KeyRangeToSet, Data.PositionalKeys, Data.RotationalKeys, Data.ScalingKeys);
+			const FInt32Range KeyRangeToSet(AnimKey, AnimKey + 1);
+			Controller.UpdateBoneTrackKeys(Data.TargetBoneName, KeyRangeToSet, { BonePoseTargetLocal.GetLocation() }, { BonePoseTargetLocal.GetRotation() }, { BonePoseTargetLocal.GetScale3D() });
+		}
 	}
 
 	// Done editing animation data
