Exposed useful functionality from MotionExtractor to allow blueprints to make their own modifiers

diff --git a/Engine/Plugins/Animation/AnimationModifierLibrary/Source/AnimationModifierLibrary/Private/MotionExtractorModifier.cpp b/Engine/Plugins/Animation/AnimationModifierLibrary/Source/AnimationModifierLibrary/Private/MotionExtractorModifier.cpp
--- a/Engine/Plugins/Animation/AnimationModifierLibrary/Source/AnimationModifierLibrary/Private/MotionExtractorModifier.cpp
+++ b/Engine/Plugins/Animation/AnimationModifierLibrary/Source/AnimationModifierLibrary/Private/MotionExtractorModifier.cpp
@@ -1,10 +1,8 @@
 // Copyright Epic Games, Inc. All Rights Reserved.
 
 #include "MotionExtractorModifier.h"
-#include "AnimationBlueprintLibrary.h"
+#include "MotionExtractorUtilities.h"
 #include "Animation/AnimSequence.h"
-#include "Animation/Skeleton.h"
-#include "BonePose.h"
 #include "Animation/AnimationPoseData.h"
 
 UMotionExtractorModifier::UMotionExtractorModifier()
@@ -68,7 +66,7 @@ void UMotionExtractorModifier::OnApply_Implementation(UAnimSequence* Animation)
 	FBoneContainer BoneContainer(RequiredBones, false, *Skeleton);
 	const FCompactPoseBoneIndex CompactPoseBoneIndex = BoneContainer.MakeCompactPoseIndex(FMeshPoseBoneIndex(BoneIndex));
 
-	FTransform FirstFrameBoneTransform = ExtractBoneTransform(Animation, BoneContainer, CompactPoseBoneIndex, 0.f, bComponentSpace);
+	FTransform FirstFrameBoneTransform = UMotionExtractorUtilityLibrary::ExtractBoneTransform(Animation, BoneContainer, CompactPoseBoneIndex, 0.f, bComponentSpace);
 
 	const float AnimLength = Animation->GetPlayLength();
 	const float SampleInterval = 1.f / SampleRate;
@@ -85,7 +83,7 @@ void UMotionExtractorModifier::OnApply_Implementation(UAnimSequence* Animation)
 			Time = FMath::Clamp(SampleIndex * SampleInterval, 0.f, AnimLength);
 			SampleIndex++;
 
-			FTransform BoneTransform = ExtractBoneTransform(Animation, BoneContainer, CompactPoseBoneIndex, Time, bComponentSpace);
+			FTransform BoneTransform = UMotionExtractorUtilityLibrary::ExtractBoneTransform(Animation, BoneContainer, CompactPoseBoneIndex, Time, bComponentSpace);
 
 			if(bRelativeToFirstFrame)
 			{
@@ -111,7 +109,7 @@ void UMotionExtractorModifier::OnApply_Implementation(UAnimSequence* Animation)
 		Time = FMath::Clamp(SampleIndex * SampleInterval, 0.f, AnimLength);
 		SampleIndex++;
 
-		FTransform BoneTransform = ExtractBoneTransform(Animation, BoneContainer, CompactPoseBoneIndex, Time, bComponentSpace);
+		FTransform BoneTransform = UMotionExtractorUtilityLibrary::ExtractBoneTransform(Animation, BoneContainer, CompactPoseBoneIndex, Time, bComponentSpace);
 
 		if (bRelativeToFirstFrame)
 		{
@@ -148,122 +146,12 @@ FName UMotionExtractorModifier::GetCurveName() const
 		return CustomCurveName;
 	}
 
-	FString MotionTypeStr;
-	switch (MotionType)
-	{
-	case EMotionExtractor_MotionType::Translation:		MotionTypeStr = TEXT("translation"); break;
-	case EMotionExtractor_MotionType::Rotation:			MotionTypeStr = TEXT("rotation"); break;
-	case EMotionExtractor_MotionType::Scale:			MotionTypeStr = TEXT("scale"); break;
-	case EMotionExtractor_MotionType::TranslationSpeed:	MotionTypeStr = TEXT("translation_speed"); break;
-	case EMotionExtractor_MotionType::RotationSpeed:	MotionTypeStr = TEXT("rotation_speed"); break;
-	default: check(false); break;
-	}
-
-	FString AxisStr;
-	switch (Axis)
-	{
-	case EMotionExtractor_Axis::X:		AxisStr = TEXT("X"); break;
-	case EMotionExtractor_Axis::Y:		AxisStr = TEXT("Y"); break;
-	case EMotionExtractor_Axis::Z:		AxisStr = TEXT("Z"); break;
-	case EMotionExtractor_Axis::XY:		AxisStr = TEXT("XY"); break;
-	case EMotionExtractor_Axis::XZ:		AxisStr = TEXT("XZ"); break;
-	case EMotionExtractor_Axis::YZ:		AxisStr = TEXT("YZ"); break;
-	case EMotionExtractor_Axis::XYZ:	AxisStr = TEXT("XYZ"); break;
-	default: check(false); break;
-	}
-
-	return FName(*FString::Printf(TEXT("%s_%s_%s"), *BoneName.ToString(), *MotionTypeStr, *AxisStr));
+	return UMotionExtractorUtilityLibrary::GenerateCurveName(BoneName, MotionType, Axis);
 }
 
 float UMotionExtractorModifier::GetDesiredValue(const FTransform& BoneTransform, const FTransform& LastBoneTransform, float DeltaTime) const
 {
-	float Value = 0.f;
-	switch (MotionType)
-	{
-	case EMotionExtractor_MotionType::Translation:
-	{
-		const FVector Translation = BoneTransform.GetTranslation();
-		switch (Axis)
-		{
-		case EMotionExtractor_Axis::X: Value = Translation.X; break;
-		case EMotionExtractor_Axis::Y: Value = Translation.Y; break;
-		case EMotionExtractor_Axis::Z: Value = Translation.Z; break;
-		default: break;
-		}
-		break;
-	}
-	case EMotionExtractor_MotionType::Rotation:
-	{
-		const FRotator Rotation = BoneTransform.GetRotation().Rotator();
-		switch (Axis)
-		{
-		case EMotionExtractor_Axis::X: Value = Rotation.Roll; break;
-		case EMotionExtractor_Axis::Y: Value = Rotation.Pitch; break;
-		case EMotionExtractor_Axis::Z: Value = Rotation.Yaw; break;
-		default: break;
-		}
-		break;
-	}
-	case EMotionExtractor_MotionType::Scale:
-	{
-		const FVector Scale = BoneTransform.GetScale3D();
-		switch (Axis)
-		{
-		case EMotionExtractor_Axis::X: Value = Scale.X; break;
-		case EMotionExtractor_Axis::Y: Value = Scale.Y; break;
-		case EMotionExtractor_Axis::Z: Value = Scale.Z; break;
-		default: break;
-		}
-		break;
-	}
-	case EMotionExtractor_MotionType::TranslationSpeed:
-	{
-		if (!FMath::IsNearlyZero(DeltaTime))
-		{
-			const float Delta = CalculateMagnitude((BoneTransform.GetTranslation() - LastBoneTransform.GetTranslation()), Axis);
-			Value = Delta / DeltaTime;
-		}
-
-		break;
-	}
-	case EMotionExtractor_MotionType::RotationSpeed:
-	{
-		if(!FMath::IsNearlyZero(DeltaTime))
-		{
-			FQuat Delta = FQuat::Identity;
-			FRotator Rotator = BoneTransform.GetRotation().Rotator();
-			FRotator LastRotator = LastBoneTransform.GetRotation().Rotator();
-
-			switch (Axis)
-			{
-			case EMotionExtractor_Axis::X:		Delta = FQuat(FRotator(0.f, 0.f, Rotator.Roll)) * FQuat(FRotator(0.f, 0.f, LastRotator.Roll)).Inverse(); break;
-			case EMotionExtractor_Axis::Y:		Delta = FQuat(FRotator(Rotator.Pitch, 0.f, 0.f)) * FQuat(FRotator(LastRotator.Pitch, 0.f, 0.f)).Inverse(); break;
-			case EMotionExtractor_Axis::Z:		Delta = FQuat(FRotator(0.f, Rotator.Yaw, 0.f)) * FQuat(FRotator(0.f, LastRotator.Yaw, 0.f)).Inverse(); break;
-			case EMotionExtractor_Axis::XY:		Delta = FQuat(FRotator(Rotator.Pitch, 0.f, Rotator.Roll)) * FQuat(FRotator(LastRotator.Pitch, 0.f, LastRotator.Roll)).Inverse(); break;
-			case EMotionExtractor_Axis::XZ:		Delta = FQuat(FRotator(0.f, Rotator.Yaw, Rotator.Roll)) * FQuat(FRotator(0.f, LastRotator.Yaw, LastRotator.Roll)).Inverse(); break;
-			case EMotionExtractor_Axis::YZ:		Delta = FQuat(FRotator(Rotator.Pitch, Rotator.Yaw, 0.f)) * FQuat(FRotator(LastRotator.Pitch, LastRotator.Yaw, 0.f)).Inverse(); break;
-			case EMotionExtractor_Axis::XYZ:	Delta = BoneTransform.GetRotation() * LastBoneTransform.GetRotation().Inverse(); break;
-			default: break;
-			}
-		
-			float RotationAngle = 0.f;
-			FVector RotationAxis;
-			Delta.ToAxisAndAngle(RotationAxis, RotationAngle);
-			RotationAngle = FMath::UnwindRadians(RotationAngle);
-			if (RotationAngle < 0.0f)
-			{
-				RotationAxis = -RotationAxis;
-				RotationAngle = -RotationAngle;
-			}
-
-			RotationAngle = FMath::RadiansToDegrees(RotationAngle);
-			Value = RotationAngle / DeltaTime;
-		}
-
-		break;
-	}
-	default: check(false); break;
-	}
+	float Value = UMotionExtractorUtilityLibrary::GetDesiredValue(BoneTransform, LastBoneTransform, DeltaTime, MotionType, Axis);
 
 	if (bAbsoluteValue)
 	{
@@ -283,49 +171,4 @@ float UMotionExtractorModifier::GetDesiredValue(const FTransform& BoneTransform,
 	}
 
 	return Value;
-}
-
-FTransform UMotionExtractorModifier::ExtractBoneTransform(UAnimSequence* Animation, const FBoneContainer& BoneContainer, FCompactPoseBoneIndex CompactPoseBoneIndex, float Time, bool bComponentSpace)
-{
-	FCompactPose Pose;
-	Pose.SetBoneContainer(&BoneContainer);
-
-	FBlendedCurve Curve;
-	Curve.InitFrom(BoneContainer);
-
-	FAnimExtractContext Context(static_cast<double>(Time), false);
-	UE::Anim::FStackAttributeContainer Attributes;
-	FAnimationPoseData AnimationPoseData(Pose, Curve, Attributes);
-
-	Animation->GetBonePose(AnimationPoseData, Context, true);
-
-	check(Pose.IsValidIndex(CompactPoseBoneIndex));
-
-	if (bComponentSpace)
-	{
-		FCSPose<FCompactPose> ComponentSpacePose;
-		ComponentSpacePose.InitPose(Pose);
-		return ComponentSpacePose.GetComponentSpaceTransform(CompactPoseBoneIndex);
-	}
-	else
-	{
-		return Pose[CompactPoseBoneIndex];
-	}
-}
-
-float UMotionExtractorModifier::CalculateMagnitude(const FVector& Vector, EMotionExtractor_Axis Axis)
-{
-	switch (Axis)
-	{
-	case EMotionExtractor_Axis::X:		return FMath::Abs(Vector.X); break;
-	case EMotionExtractor_Axis::Y:		return FMath::Abs(Vector.Y); break;
-	case EMotionExtractor_Axis::Z:		return FMath::Abs(Vector.Z); break;
-	case EMotionExtractor_Axis::XY:		return FMath::Sqrt(Vector.X * Vector.X + Vector.Y * Vector.Y); break;
-	case EMotionExtractor_Axis::XZ:		return FMath::Sqrt(Vector.X * Vector.X + Vector.Z * Vector.Z); break;
-	case EMotionExtractor_Axis::YZ:		return FMath::Sqrt(Vector.Y * Vector.Y + Vector.Z * Vector.Z); break;
-	case EMotionExtractor_Axis::XYZ:	return FMath::Sqrt(Vector.X * Vector.X + Vector.Y * Vector.Y + Vector.Z * Vector.Z); break;
-	default: check(false); break;
-	}
-
-	return 0.f;
 }
\ No newline at end of file
diff --git a/Engine/Plugins/Animation/AnimationModifierLibrary/Source/AnimationModifierLibrary/Private/MotionExtractorUtilities.cpp b/Engine/Plugins/Animation/AnimationModifierLibrary/Source/AnimationModifierLibrary/Private/MotionExtractorUtilities.cpp
new file mode 100644
--- /dev/null
+++ b/Engine/Plugins/Animation/AnimationModifierLibrary/Source/AnimationModifierLibrary/Private/MotionExtractorUtilities.cpp
@@ -0,0 +1,186 @@
+// Copyright Epic Games, Inc. All Rights Reserved.
+
+#include "MotionExtractorUtilities.h"
+#include "Animation/AnimSequence.h"
+#include "Animation/AnimationPoseData.h"
+
+FName UMotionExtractorUtilityLibrary::GenerateCurveName(
+	FName BoneName,
+	EMotionExtractor_MotionType MotionType,
+	EMotionExtractor_Axis Axis)
+{
+	FString MotionTypeStr;
+	switch (MotionType)
+	{
+	case EMotionExtractor_MotionType::Translation:		MotionTypeStr = TEXT("translation"); break;
+	case EMotionExtractor_MotionType::Rotation:			MotionTypeStr = TEXT("rotation"); break;
+	case EMotionExtractor_MotionType::Scale:			MotionTypeStr = TEXT("scale"); break;
+	case EMotionExtractor_MotionType::TranslationSpeed:	MotionTypeStr = TEXT("translation_speed"); break;
+	case EMotionExtractor_MotionType::RotationSpeed:	MotionTypeStr = TEXT("rotation_speed"); break;
+	default: check(false); break;
+	}
+
+	FString AxisStr;
+	switch (Axis)
+	{
+	case EMotionExtractor_Axis::X:		AxisStr = TEXT("X"); break;
+	case EMotionExtractor_Axis::Y:		AxisStr = TEXT("Y"); break;
+	case EMotionExtractor_Axis::Z:		AxisStr = TEXT("Z"); break;
+	case EMotionExtractor_Axis::XY:		AxisStr = TEXT("XY"); break;
+	case EMotionExtractor_Axis::XZ:		AxisStr = TEXT("XZ"); break;
+	case EMotionExtractor_Axis::YZ:		AxisStr = TEXT("YZ"); break;
+	case EMotionExtractor_Axis::XYZ:	AxisStr = TEXT("XYZ"); break;
+	default: check(false); break;
+	}
+
+	return FName(*FString::Printf(TEXT("%s_%s_%s"), *BoneName.ToString(), *MotionTypeStr, *AxisStr));
+}
+
+float UMotionExtractorUtilityLibrary::GetDesiredValue(
+	const FTransform& BoneTransform, 
+	const FTransform& LastBoneTransform, 
+	float DeltaTime, 
+	EMotionExtractor_MotionType MotionType,
+	EMotionExtractor_Axis Axis)
+{
+	if ((MotionType == EMotionExtractor_MotionType::Translation || MotionType == EMotionExtractor_MotionType::Rotation || MotionType == EMotionExtractor_MotionType::Scale) && Axis > EMotionExtractor_Axis::Z)
+	{
+		UE_LOG(LogAnimation, Error, TEXT("MotionExtractorUtilityLibrary failed. Reason: Only X, Y or Z axes are valid options for the selected motion type"));
+		return 0.0f;
+	}
+
+	float Value = 0.f;
+	switch (MotionType)
+	{
+	case EMotionExtractor_MotionType::Translation:
+	{
+		const FVector Translation = BoneTransform.GetTranslation();
+		switch (Axis)
+		{
+		case EMotionExtractor_Axis::X: Value = Translation.X; break;
+		case EMotionExtractor_Axis::Y: Value = Translation.Y; break;
+		case EMotionExtractor_Axis::Z: Value = Translation.Z; break;
+		default: break;
+		}
+		break;
+	}
+	case EMotionExtractor_MotionType::Rotation:
+	{
+		const FRotator Rotation = BoneTransform.GetRotation().Rotator();
+		switch (Axis)
+		{
+		case EMotionExtractor_Axis::X: Value = Rotation.Roll; break;
+		case EMotionExtractor_Axis::Y: Value = Rotation.Pitch; break;
+		case EMotionExtractor_Axis::Z: Value = Rotation.Yaw; break;
+		default: break;
+		}
+		break;
+	}
+	case EMotionExtractor_MotionType::Scale:
+	{
+		const FVector Scale = BoneTransform.GetScale3D();
+		switch (Axis)
+		{
+		case EMotionExtractor_Axis::X: Value = Scale.X; break;
+		case EMotionExtractor_Axis::Y: Value = Scale.Y; break;
+		case EMotionExtractor_Axis::Z: Value = Scale.Z; break;
+		default: break;
+		}
+		break;
+	}
+	case EMotionExtractor_MotionType::TranslationSpeed:
+	{
+		if (!FMath::IsNearlyZero(DeltaTime))
+		{
+			const float Delta = UMotionExtractorUtilityLibrary::CalculateMagnitude((BoneTransform.GetTranslation() - LastBoneTransform.GetTranslation()), Axis);
+			Value = Delta / DeltaTime;
+		}
+
+		break;
+	}
+	case EMotionExtractor_MotionType::RotationSpeed:
+	{
+		if (!FMath::IsNearlyZero(DeltaTime))
+		{
+			FQuat Delta = FQuat::Identity;
+			FRotator Rotator = BoneTransform.GetRotation().Rotator();
+			FRotator LastRotator = LastBoneTransform.GetRotation().Rotator();
+
+			switch (Axis)
+			{
+			case EMotionExtractor_Axis::X:		Delta = FQuat(FRotator(0.f, 0.f, Rotator.Roll)) * FQuat(FRotator(0.f, 0.f, LastRotator.Roll)).Inverse(); break;
+			case EMotionExtractor_Axis::Y:		Delta = FQuat(FRotator(Rotator.Pitch, 0.f, 0.f)) * FQuat(FRotator(LastRotator.Pitch, 0.f, 0.f)).Inverse(); break;
+			case EMotionExtractor_Axis::Z:		Delta = FQuat(FRotator(0.f, Rotator.Yaw, 0.f)) * FQuat(FRotator(0.f, LastRotator.Yaw, 0.f)).Inverse(); break;
+			case EMotionExtractor_Axis::XY:		Delta = FQuat(FRotator(Rotator.Pitch, 0.f, Rotator.Roll)) * FQuat(FRotator(LastRotator.Pitch, 0.f, LastRotator.Roll)).Inverse(); break;
+			case EMotionExtractor_Axis::XZ:		Delta = FQuat(FRotator(0.f, Rotator.Yaw, Rotator.Roll)) * FQuat(FRotator(0.f, LastRotator.Yaw, LastRotator.Roll)).Inverse(); break;
+			case EMotionExtractor_Axis::YZ:		Delta = FQuat(FRotator(Rotator.Pitch, Rotator.Yaw, 0.f)) * FQuat(FRotator(LastRotator.Pitch, LastRotator.Yaw, 0.f)).Inverse(); break;
+			case EMotionExtractor_Axis::XYZ:	Delta = BoneTransform.GetRotation() * LastBoneTransform.GetRotation().Inverse(); break;
+			default: break;
+			}
+
+			float RotationAngle = 0.f;
+			FVector RotationAxis;
+			Delta.ToAxisAndAngle(RotationAxis, RotationAngle);
+			RotationAngle = FMath::UnwindRadians(RotationAngle);
+			if (RotationAngle < 0.0f)
+			{
+				RotationAxis = -RotationAxis;
+				RotationAngle = -RotationAngle;
+			}
+
+			RotationAngle = FMath::RadiansToDegrees(RotationAngle);
+			Value = RotationAngle / DeltaTime;
+		}
+
+		break;
+	}
+	default: check(false); break;
+	}
+
+	return Value;
+}
+
+FTransform UMotionExtractorUtilityLibrary::ExtractBoneTransform(UAnimSequence* Animation, const FBoneContainer& BoneContainer, FCompactPoseBoneIndex CompactPoseBoneIndex, float Time, bool bComponentSpace)
+{
+	FCompactPose Pose;
+	Pose.SetBoneContainer(&BoneContainer);
+
+	FBlendedCurve Curve;
+	Curve.InitFrom(BoneContainer);
+
+	FAnimExtractContext Context(Time, false);
+	UE::Anim::FStackAttributeContainer Attributes;
+	FAnimationPoseData AnimationPoseData(Pose, Curve, Attributes);
+
+	Animation->GetBonePose(AnimationPoseData, Context, true);
+
+	check(Pose.IsValidIndex(CompactPoseBoneIndex));
+
+	if (bComponentSpace)
+	{
+		FCSPose<FCompactPose> ComponentSpacePose;
+		ComponentSpacePose.InitPose(Pose);
+		return ComponentSpacePose.GetComponentSpaceTransform(CompactPoseBoneIndex);
+	}
+	else
+	{
+		return Pose[CompactPoseBoneIndex];
+	}
+}
+
+float UMotionExtractorUtilityLibrary::CalculateMagnitude(const FVector& Vector, EMotionExtractor_Axis Axis)
+{
+	switch (Axis)
+	{
+	case EMotionExtractor_Axis::X:		return FMath::Abs(Vector.X); break;
+	case EMotionExtractor_Axis::Y:		return FMath::Abs(Vector.Y); break;
+	case EMotionExtractor_Axis::Z:		return FMath::Abs(Vector.Z); break;
+	case EMotionExtractor_Axis::XY:		return FMath::Sqrt(Vector.X * Vector.X + Vector.Y * Vector.Y); break;
+	case EMotionExtractor_Axis::XZ:		return FMath::Sqrt(Vector.X * Vector.X + Vector.Z * Vector.Z); break;
+	case EMotionExtractor_Axis::YZ:		return FMath::Sqrt(Vector.Y * Vector.Y + Vector.Z * Vector.Z); break;
+	case EMotionExtractor_Axis::XYZ:	return FMath::Sqrt(Vector.X * Vector.X + Vector.Y * Vector.Y + Vector.Z * Vector.Z); break;
+	default: check(false); break;
+	}
+
+	return 0.f;
+}
\ No newline at end of file
diff --git a/Engine/Plugins/Animation/AnimationModifierLibrary/Source/AnimationModifierLibrary/Public/MotionExtractorModifier.h b/Engine/Plugins/Animation/AnimationModifierLibrary/Source/AnimationModifierLibrary/Public/MotionExtractorModifier.h
--- a/Engine/Plugins/Animation/AnimationModifierLibrary/Source/AnimationModifierLibrary/Public/MotionExtractorModifier.h
+++ b/Engine/Plugins/Animation/AnimationModifierLibrary/Source/AnimationModifierLibrary/Public/MotionExtractorModifier.h
@@ -4,43 +4,9 @@
 
 #include "CoreMinimal.h"
 #include "AnimationModifier.h"
+#include "MotionExtractorTypes.h"
 #include "MotionExtractorModifier.generated.h"
 
-/** Type of motion to extract */
-UENUM(BlueprintType)
-enum class EMotionExtractor_MotionType : uint8
-{
-	Translation,
-	Rotation,
-	Scale,
-	TranslationSpeed,
-	RotationSpeed
-};
-
-/** Axis to get the final value from */
-UENUM(BlueprintType)
-enum class EMotionExtractor_Axis : uint8
-{
-	X,
-	Y,
-	Z,
-	XY,
-	XZ,
-	YZ,
-	XYZ
-};
-
-/** Math operations that can be applied to the extracted value before add it to the curve */
-UENUM(BlueprintType)
-enum class EMotionExtractor_MathOperation : uint8
-{
-	None,
-	Addition,
-	Subtraction,
-	Division,
-	Multiplication
-};
-
 /** Extracts motion from a bone in the animation and bakes it into a curve */
 UCLASS()
 class UMotionExtractorModifier : public UAnimationModifier
@@ -113,10 +79,4 @@ public:
 
 	/** Returns the desired value from the extracted poses */
 	float GetDesiredValue(const FTransform& BoneTransform, const FTransform& LastBoneTransform, float DeltaTime) const;
-
-	/** Helper function to extract the pose for a given bone at a given time */
-	static FTransform ExtractBoneTransform(UAnimSequence* Animation, const FBoneContainer& BoneContainer, FCompactPoseBoneIndex CompactPoseBoneIndex, float Time, bool bComponentSpace);
-
-	/** Helper function to calculate the magnitude of a vector only considering a specific axis or axes */
-	static float CalculateMagnitude(const FVector& Vector, EMotionExtractor_Axis Axis);
 };
diff --git a/Engine/Plugins/Animation/AnimationModifierLibrary/Source/AnimationModifierLibrary/Public/MotionExtractorTypes.h b/Engine/Plugins/Animation/AnimationModifierLibrary/Source/AnimationModifierLibrary/Public/MotionExtractorTypes.h
new file mode 100644
--- /dev/null
+++ b/Engine/Plugins/Animation/AnimationModifierLibrary/Source/AnimationModifierLibrary/Public/MotionExtractorTypes.h
@@ -0,0 +1,42 @@
+// Copyright Epic Games, Inc. All Rights Reserved.
+
+#pragma once
+
+#include "CoreMinimal.h"
+#include "AnimationBlueprintLibrary.h"
+#include "MotionExtractorTypes.generated.h"
+
+/** Type of motion to extract */
+UENUM(BlueprintType)
+enum class EMotionExtractor_MotionType : uint8
+{
+	Translation,
+	Rotation,
+	Scale,
+	TranslationSpeed,
+	RotationSpeed
+};
+
+/** Axis to get the final value from */
+UENUM(BlueprintType)
+enum class EMotionExtractor_Axis : uint8
+{
+	X,
+	Y,
+	Z,
+	XY,
+	XZ,
+	YZ,
+	XYZ
+};
+
+/** Math operations that can be applied to the extracted value before add it to the curve */
+UENUM(BlueprintType)
+enum class EMotionExtractor_MathOperation : uint8
+{
+	None,
+	Addition,
+	Subtraction,
+	Division,
+	Multiplication
+};
\ No newline at end of file
diff --git a/Engine/Plugins/Animation/AnimationModifierLibrary/Source/AnimationModifierLibrary/Public/MotionExtractorUtilities.h b/Engine/Plugins/Animation/AnimationModifierLibrary/Source/AnimationModifierLibrary/Public/MotionExtractorUtilities.h
new file mode 100644
--- /dev/null
+++ b/Engine/Plugins/Animation/AnimationModifierLibrary/Source/AnimationModifierLibrary/Public/MotionExtractorUtilities.h
@@ -0,0 +1,54 @@
+// Copyright Epic Games, Inc. All Rights Reserved.
+
+#pragma once
+
+#include "Kismet/BlueprintFunctionLibrary.h"
+#include "MotionExtractorTypes.h"
+#include "MotionExtractorUtilities.generated.h"
+
+UCLASS(meta=(ScriptName="MotionExtractorUtilityLibrary"))
+class  UMotionExtractorUtilityLibrary : public UBlueprintFunctionLibrary
+{
+	GENERATED_BODY()
+
+public:
+
+	/**
+	* Generates a curve name based on input settings.
+	*
+	* @param BoneName			   The name of the bone
+	* @param MotionType            What type of motion this curve represents (translation, rotation, speed, etc.)
+	* @param Axis				   Which axis/axes the motion for this curve was extracted from
+	*/
+	UFUNCTION(BlueprintCallable, Category = "Motion Extractor Utility")
+	static FName GenerateCurveName(
+		FName BoneName,
+		EMotionExtractor_MotionType MotionType,
+		EMotionExtractor_Axis Axis);
+
+	/**
+	* Returns the desired value from the extracted poses based on input settings.
+	*
+	* @param BoneTransform         Current frame's bone transform
+	* @param LastBoneTransform     Last frame's bone transform. Unused when not calculating speeds.
+	* @param DeltaTime			   Time step used between current and last bone transforms. Unused when not calculating speeds.
+	* @param MotionType            What type of motion to extract (translation, rotation, speed, etc.)
+	* @param Axis				   Which axis/axes to extract motion from
+	*/
+	UFUNCTION(BlueprintCallable, Category = "Motion Extractor Utility")
+	static float GetDesiredValue(
+		const FTransform& BoneTransform,
+		const FTransform& LastBoneTransform,
+		float DeltaTime,
+		EMotionExtractor_MotionType MotionType,
+		EMotionExtractor_Axis Axis);
+
+	/** 
+	* Helper function to extract the pose for a given bone at a given time
+	* IMPORTANT: This function expects you to add a MemMark (FMemMark Mark(FMemStack::Get());) at the correct scope if you are using it from outside world's tick
+	*/
+	static FTransform ExtractBoneTransform(UAnimSequence* Animation, const FBoneContainer& BoneContainer, FCompactPoseBoneIndex CompactPoseBoneIndex, float Time, bool bComponentSpace);
+
+	/** Helper function to calculate the magnitude of a vector only considering a specific axis or axes */
+	static float CalculateMagnitude(const FVector& Vector, EMotionExtractor_Axis Axis);
+};
\ No newline at end of file
