Fix for build error on deprecated FAnimContext use

diff --git a/Engine/Plugins/Animation/AnimationModifierLibrary/Source/AnimationModifierLibrary/Private/MotionExtractorUtilities.cpp b/Engine/Plugins/Animation/AnimationModifierLibrary/Source/AnimationModifierLibrary/Private/MotionExtractorUtilities.cpp
--- a/Engine/Plugins/Animation/AnimationModifierLibrary/Source/AnimationModifierLibrary/Private/MotionExtractorUtilities.cpp
+++ b/Engine/Plugins/Animation/AnimationModifierLibrary/Source/AnimationModifierLibrary/Private/MotionExtractorUtilities.cpp
@@ -148,7 +148,7 @@ FTransform UMotionExtractorUtilityLibrary::ExtractBoneTransform(UAnimSequence* A
 	FBlendedCurve Curve;
 	Curve.InitFrom(BoneContainer);
 
-	FAnimExtractContext Context(Time, false);
+	FAnimExtractContext Context(static_cast<double>(Time), false);
 	UE::Anim::FStackAttributeContainer Attributes;
 	FAnimationPoseData AnimationPoseData(Pose, Curve, Attributes);
 
