-Added blueprint function for custom animation modifiers to find stop/moving ranges
-Fixed stop detection in orientation warping modifier
#rnx

diff --git a/Engine/Plugins/Animation/AnimationModifierLibrary/Source/AnimationModifierLibrary/Private/MotionExtractorUtilities.cpp b/Engine/Plugins/Animation/AnimationModifierLibrary/Source/AnimationModifierLibrary/Private/MotionExtractorUtilities.cpp
--- a/Engine/Plugins/Animation/AnimationModifierLibrary/Source/AnimationModifierLibrary/Private/MotionExtractorUtilities.cpp
+++ b/Engine/Plugins/Animation/AnimationModifierLibrary/Source/AnimationModifierLibrary/Private/MotionExtractorUtilities.cpp
@@ -4,6 +4,51 @@
 #include "Animation/AnimSequence.h"
 #include "Animation/AnimationPoseData.h"
 
+namespace UE::Anim::MotionExtractorUtility
+{
+	// Returns the ranges (X/Start to Y/End) in the specified animation sequence where the passes the specified criterion
+	template<typename Predicate>
+	TArray<FVector2D> GetRangesFromRootMotionByPredicate(const UAnimSequence* Animation, Predicate Pred, float SampleRate = 120.0f)
+	{
+		if (Animation == nullptr)
+		{
+			return TArray<FVector2D>();
+		}
+
+		const float AnimLength = Animation->GetPlayLength();
+		float SampleInterval = 1.f / SampleRate;
+		int32 NumSteps = AnimLength / SampleInterval;
+
+		TArray<FVector2D> OutRanges;
+		OutRanges.Reserve(4);
+		bool bMeetsCriterion = false;
+		for (int32 Step = 0; Step < NumSteps; ++Step)
+		{
+			const float Time = Step * SampleInterval;
+			const bool bAllowLooping = false;
+			const FTransform RootMotion = Animation->ExtractRootMotion(Time, SampleInterval, bAllowLooping);
+			const bool bMetCriterion = bMeetsCriterion;
+			bMeetsCriterion = ::Invoke(Pred, RootMotion, SampleInterval);
+
+			if (bMeetsCriterion != bMetCriterion)
+			{
+				if (bMeetsCriterion)
+				{
+					// Assume the range lasts from now until the animation ends.
+					OutRanges.Add(FVector2D(Time, AnimLength));
+				}
+				else if (!OutRanges.IsEmpty())
+				{
+					// No longer meets the criterion. Close the range.
+					OutRanges.Top().Y = Time;
+				}
+			}
+		}
+
+		return OutRanges;
+	}
+}
+
 FName UMotionExtractorUtilityLibrary::GenerateCurveName(
 	FName BoneName,
 	EMotionExtractor_MotionType MotionType,
@@ -140,6 +185,30 @@ float UMotionExtractorUtilityLibrary::GetDesiredValue(
 	return Value;
 }
 
+TArray<FVector2D> UMotionExtractorUtilityLibrary::GetStoppedRangesFromRootMotion(const UAnimSequence* AnimSequence, float StopSpeedThreshold, float SampleRate)
+{
+	auto Criterion = [StopSpeedThreshold](const FTransform& InRootMotionDelta, float DeltaTime)
+	{
+		const FVector RootMotionTranslation = InRootMotionDelta.GetTranslation();
+		const float RootMotionSpeed = RootMotionTranslation.Size() / DeltaTime;
+		return RootMotionSpeed < StopSpeedThreshold;
+	};
+
+	return UE::Anim::MotionExtractorUtility::GetRangesFromRootMotionByPredicate(AnimSequence, Criterion, SampleRate);
+}
+
+TArray<FVector2D> UMotionExtractorUtilityLibrary::GetMovingRangesFromRootMotion(const UAnimSequence* AnimSequence, float StopSpeedThreshold /*= 10.0f*/, float SampleRate /*= 120.0f*/)
+{
+	auto Criterion = [StopSpeedThreshold](const FTransform& InRootMotionDelta, float DeltaTime)
+	{
+		const FVector RootMotionTranslation = InRootMotionDelta.GetTranslation();
+		const float RootMotionSpeed = RootMotionTranslation.Size() / DeltaTime;
+		return RootMotionSpeed > StopSpeedThreshold;
+	};
+
+	return UE::Anim::MotionExtractorUtility::GetRangesFromRootMotionByPredicate(AnimSequence, Criterion, SampleRate);
+}
+
 FTransform UMotionExtractorUtilityLibrary::ExtractBoneTransform(UAnimSequence* Animation, const FBoneContainer& BoneContainer, FCompactPoseBoneIndex CompactPoseBoneIndex, float Time, bool bComponentSpace)
 {
 	FCompactPose Pose;
@@ -183,4 +252,4 @@ float UMotionExtractorUtilityLibrary::CalculateMagnitude(const FVector& Vector,
 	}
 
 	return 0.f;
-}
\ No newline at end of file
+}
diff --git a/Engine/Plugins/Animation/AnimationModifierLibrary/Source/AnimationModifierLibrary/Public/MotionExtractorUtilities.h b/Engine/Plugins/Animation/AnimationModifierLibrary/Source/AnimationModifierLibrary/Public/MotionExtractorUtilities.h
--- a/Engine/Plugins/Animation/AnimationModifierLibrary/Source/AnimationModifierLibrary/Public/MotionExtractorUtilities.h
+++ b/Engine/Plugins/Animation/AnimationModifierLibrary/Source/AnimationModifierLibrary/Public/MotionExtractorUtilities.h
@@ -43,6 +43,26 @@ public:
 		EMotionExtractor_MotionType MotionType,
 		EMotionExtractor_Axis Axis);
 
+	/**
+	* Returns the ranges (X/Start to Y/End) in the specified animation sequence where the animation is considered stopped.
+	*
+	* @param AnimSequence			Anim sequence to check
+	* @param StopSpeedThreshold		Root motion speed under which the animation is considered stopped.
+	* @param SampleRate				Sample rate of the animation. It's recommended to use high values if the animation has very sudden direction changes.
+	*/
+	UFUNCTION(BlueprintCallable, Category = "Motion Extractor Utility")
+	static TArray<FVector2D> GetStoppedRangesFromRootMotion(const UAnimSequence* AnimSequence, float StopSpeedThreshold = 10.0f, float SampleRate = 120.0f);
+
+	/**
+	* Returns the ranges (X/Start to Y/End) in the specified animation sequence where the animation is considered moving.
+	*
+	* @param AnimSequence			Anim sequence to check
+	* @param StopSpeedThreshold		Root motion speed over which the animation is considered moving.
+	* @param SampleRate				Sample rate of the animation. It's recommended to use high values if the animation has very sudden direction changes.
+	*/
+	UFUNCTION(BlueprintCallable, Category = "Motion Extractor Utility")
+	static TArray<FVector2D> GetMovingRangesFromRootMotion(const UAnimSequence* AnimSequence, float StopSpeedThreshold = 10.0f, float SampleRate = 120.0f);
+
 	/** 
 	* Helper function to extract the pose for a given bone at a given time
 	* IMPORTANT: This function expects you to add a MemMark (FMemMark Mark(FMemStack::Get());) at the correct scope if you are using it from outside world's tick
diff --git a/Engine/Plugins/Animation/AnimationWarping/Source/Editor/Private/AnimationModifiers/OrientationWarpingModifier.cpp b/Engine/Plugins/Animation/AnimationWarping/Source/Editor/Private/AnimationModifiers/OrientationWarpingModifier.cpp
--- a/Engine/Plugins/Animation/AnimationWarping/Source/Editor/Private/AnimationModifiers/OrientationWarpingModifier.cpp
+++ b/Engine/Plugins/Animation/AnimationWarping/Source/Editor/Private/AnimationModifiers/OrientationWarpingModifier.cpp
@@ -38,7 +38,6 @@ void UOrientationWarpingModifier::OnApply_Implementation(UAnimSequence* Animatio
 	const float AnimLength = Animation->GetPlayLength();
 	float SampleInterval = 1.f / 120.f;
 	int32 NumSteps = AnimLength / SampleInterval;
-	float MinSpeedSq = FMath::Square(StopSpeedThreshold);
 
 	bool bIsRotating = false;
 	bool bFoundStop = false;
@@ -58,9 +57,9 @@ void UOrientationWarpingModifier::OnApply_Implementation(UAnimSequence* Animatio
 		FVector RotationAxis;
 		float RotationAngle;
 		RootMotionRotation.ToAxisAndAngle(RotationAxis, RotationAngle);
-		const float RootMotionSpeedSq = RootMotionTranslation.SizeSquared2D() / SampleInterval;
+		const float RootMotionSpeed = RootMotionTranslation.Size2D() / SampleInterval;
 
-		const bool bNewIsMoving = RootMotionSpeedSq > MinSpeedSq;
+		const bool bNewIsMoving = RootMotionSpeed > StopSpeedThreshold;
 		const bool bNewIsRotating = FMath::Abs(RotationAngle) > KINDA_SMALL_NUMBER;
 
 		// Only add keys at the beginning/end of rotations
