Crash when applying a Motion Extractor Modifier to any animation
#fix Delay adding motion curve/keys until the end of MotionExtractorModifier processing. Recently ::OnApply of an AnimationModifier has been moved to happen within a IAnimationDataController::ScopedBracket. Which means that any changes to the animation data (notifies) are only propagated at the end. The Modifier in question, added a curve and directly afterwards tried to sample the raw data (which fails)

diff --git a/Engine/Plugins/Animation/AnimationModifierLibrary/Source/AnimationModifierLibrary/Private/MotionExtractorModifier.cpp b/Engine/Plugins/Animation/AnimationModifierLibrary/Source/AnimationModifierLibrary/Private/MotionExtractorModifier.cpp
--- a/Engine/Plugins/Animation/AnimationModifierLibrary/Source/AnimationModifierLibrary/Private/MotionExtractorModifier.cpp
+++ b/Engine/Plugins/Animation/AnimationModifierLibrary/Source/AnimationModifierLibrary/Private/MotionExtractorModifier.cpp
@@ -61,9 +61,6 @@ void UMotionExtractorModifier::OnApply_Implementation(UAnimSequence* Animation)
 
 	TGuardValue<bool> ForceRootLockGuard(Animation->bForceRootLock, false);
 
-	const FName FinalCurveName = GetCurveName();
-	UAnimationBlueprintLibrary::AddCurve(Animation, FinalCurveName, ERawCurveTrackTypes::RCT_Float, false);
-
 	TArray<FBoneIndexType> RequiredBones;
 	RequiredBones.Add(BoneIndex);
 	Skeleton->GetReferenceSkeleton().EnsureParentsExistAndSort(RequiredBones);
@@ -109,6 +106,9 @@ void UMotionExtractorModifier::OnApply_Implementation(UAnimSequence* Animation)
 	LastBoneTransform = FTransform::Identity;
 	Time = 0.f;
 	SampleIndex = 0;
+
+	TArray<FRichCurveKey> CurveKeys;
+	
 	while (Time < AnimLength)
 	{
 		Time = FMath::Clamp(SampleIndex * SampleInterval, 0.f, AnimLength);
@@ -127,11 +127,20 @@ void UMotionExtractorModifier::OnApply_Implementation(UAnimSequence* Animation)
 			const float Value = GetDesiredValue(BoneTransform, LastBoneTransform, SampleInterval);
 			const float FinalValue = (bNormalize && MaxValue != 0.f) ? FMath::Abs(Value) / MaxValue : Value;
 
-			UAnimationBlueprintLibrary::AddFloatCurveKey(Animation, FinalCurveName, Time, FinalValue);
+			FRichCurveKey& Key = CurveKeys.AddDefaulted_GetRef();
+			Key.Time = Time;
+			Key.Value = FinalValue;
 		}
 
 		LastBoneTransform = BoneTransform;
 	}
+
+	IAnimationDataController& Controller = Animation->GetController();
+	const FAnimationCurveIdentifier CurveId = UAnimationCurveIdentifierExtensions::GetCurveIdentifier(Animation->GetSkeleton(), GetCurveName(), ERawCurveTrackTypes::RCT_Float);
+	if(CurveKeys.Num() && Controller.AddCurve(CurveId))
+	{
+		Controller.SetCurveKeys(CurveId, CurveKeys);
+	}
 }
 
 void UMotionExtractorModifier::OnRevert_Implementation(UAnimSequence* Animation)
