UReOrientRootBoneModifier is causing the root bone spinning uncontrollably
#fix Cache out new bone track keys, and set all at once, rather than trying to sample directly from the data (which changes with every loop cycle)

diff --git a/Engine/Plugins/Animation/AnimationModifierLibrary/Source/AnimationModifierLibrary/Private/ReOrientRootBoneModifier.cpp b/Engine/Plugins/Animation/AnimationModifierLibrary/Source/AnimationModifierLibrary/Private/ReOrientRootBoneModifier.cpp
--- a/Engine/Plugins/Animation/AnimationModifierLibrary/Source/AnimationModifierLibrary/Private/ReOrientRootBoneModifier.cpp
+++ b/Engine/Plugins/Animation/AnimationModifierLibrary/Source/AnimationModifierLibrary/Private/ReOrientRootBoneModifier.cpp
@@ -76,6 +76,23 @@ void UReOrientRootBoneModifier::ReOrientRootBone_Internal(UAnimSequence* Animati
 	const bool bShouldTransact = false;
 	Controller.OpenBracket(LOCTEXT("ReorientingRootBone_Bracket", "Reorienting root bone"), bShouldTransact);
 
+	struct FBoneKeys
+	{
+		TArray<FVector> Location;
+		TArray<FQuat> Rotation;
+		TArray<FVector> Scale;
+
+		void Add(const FTransform& Transform)
+		{
+			Location.Add(Transform.GetLocation());
+			Rotation.Add(Transform.GetRotation());
+			Scale.Add(Transform.GetScale3D());
+		}		
+	};
+
+	FBoneKeys NewRootKeys;
+	TMap<FName, FBoneKeys> NewChildKeys;
+		
 	// For each key in the animation
 	const int32 Num = Model->GetNumberOfKeys();
 	for (int32 AnimKey = 0; AnimKey < Num; AnimKey++)
@@ -87,19 +104,28 @@ void UReOrientRootBoneModifier::ReOrientRootBone_Internal(UAnimSequence* Animati
 
 		FTransform RootTransformNew = RootTransformOriginal;
 		RootTransformNew.SetRotation(RootTransformOriginal.GetRotation() * Quat);
-		Controller.UpdateBoneTrackKeys(RootBoneName, KeyRangeToSet, { RootTransformNew.GetLocation() }, { RootTransformNew.GetRotation() }, { RootTransformNew.GetScale3D() });
-
+		NewRootKeys.Add(RootTransformNew);
+		
 		// Now the mesh is facing the wrong axis. Update direct children of the root with the local space transform that puts them back to where they were originally
-
 		for (const FName& ChildBoneName : RootBoneChildBones)
 		{
+			FBoneKeys& ChildKeys = NewChildKeys.FindOrAdd(ChildBoneName);
+			
 			FTransform OtherBoneTransform = Model->EvaluateBoneTrackTransform(ChildBoneName, AnimKey, EAnimInterpolationType::Step);
 			OtherBoneTransform = OtherBoneTransform * RootTransformOriginal;
 			OtherBoneTransform = OtherBoneTransform.GetRelativeTransform(RootTransformNew);
-			Controller.UpdateBoneTrackKeys(ChildBoneName, KeyRangeToSet, { OtherBoneTransform.GetLocation() }, { OtherBoneTransform.GetRotation() }, { OtherBoneTransform.GetScale3D() });
+			ChildKeys.Add(OtherBoneTransform);
 		}
 	}
 
+	Controller.SetBoneTrackKeys(RootBoneName, NewRootKeys.Location, NewRootKeys.Rotation, NewRootKeys.Scale);
+	
+	for (const FName& ChildBoneName : RootBoneChildBones)
+	{
+		const FBoneKeys& ChildKeys = NewChildKeys.FindChecked(ChildBoneName);
+		Controller.SetBoneTrackKeys(ChildBoneName, ChildKeys.Location, ChildKeys.Rotation, ChildKeys.Scale);
+	}
+
 	// Done editing animation data
 	Controller.CloseBracket(bShouldTransact);
 }
