AnimationModifierLibrary - EncodeRootBoneModifier

diff --git a/Engine/Plugins/Animation/AnimationModifierLibrary/Source/AnimationModifierLibrary/Private/EncodeRootBoneModifier.cpp b/Engine/Plugins/Animation/AnimationModifierLibrary/Source/AnimationModifierLibrary/Private/EncodeRootBoneModifier.cpp
new file mode 100644
--- /dev/null
+++ b/Engine/Plugins/Animation/AnimationModifierLibrary/Source/AnimationModifierLibrary/Private/EncodeRootBoneModifier.cpp
@@ -0,0 +1,162 @@
+// Copyright Epic Games, Inc. All Rights Reserved.
+
+#if WITH_EDITOR
+
+#define ENABLE_ENCODE_ROOT_BONE_MODIFIER_TEST 0
+
+#include "EncodeRootBoneModifier.h"
+#include "Animation/AnimSequence.h"
+#include "PoseExtractorUtilityLibrary.h"
+
+#define LOCTEXT_NAMESPACE "EncodeRootBoneModifier"
+
+void UEncodeRootBoneModifier::OnApply_Implementation(UAnimSequence* Animation)
+{
+	FPoseExtractorUtilityLibrary Extractor(Animation);
+
+	float TotalWeightToComputeRootPosition = 0.f;
+	for (const FEncodeRootBoneWeightedBone& WeightedBone : WeightedBoneToComputeRootPosition)
+	{
+		if (!Extractor.AddBone(WeightedBone.Bone))
+		{
+			return;
+		}
+		TotalWeightToComputeRootPosition += WeightedBone.Weight;
+	}
+
+	float TotalWeightToComputeRootOrientation = 0.f;
+	for (const FEncodeRootBoneWeightedBone& WeightedBone : WeightedBoneToComputeRootOrientation)
+	{
+		if (!Extractor.AddBone(WeightedBone.Bone))
+		{
+			return;
+		}
+		TotalWeightToComputeRootOrientation += WeightedBone.Weight;
+	}
+
+	// Collect tracks for the bones we are going to modify (root and direct children of the root)
+	TArray<FName> RootBoneChildBones;
+	TArray<int32> ChildBoneIndices;
+	Animation->GetSkeleton()->GetReferenceSkeleton().GetDirectChildBones(0, ChildBoneIndices);
+	for (const int32 ChildBoneIndex : ChildBoneIndices)
+	{
+		const FName ChildBoneName = Animation->GetSkeleton()->GetReferenceSkeleton().GetBoneName(ChildBoneIndex);
+		if (Animation->GetDataModel()->IsValidBoneTrackName(ChildBoneName))
+		{
+			RootBoneChildBones.Add(ChildBoneName);
+			if (!Extractor.AddBone(ChildBoneName, ChildBoneIndex))
+			{
+				return;
+			}
+		}
+	}
+
+	if (!Extractor.CacheGlobalBoneTransforms())
+	{
+		return;
+	}
+
+	// Start editing animation data
+	const bool bShouldTransact = false;
+	Animation->GetController().OpenBracket(LOCTEXT("EncodeRootBoneModifier_Bracket", "Encoding root bone"), bShouldTransact);
+
+	// For each key in the animation
+	const int32 NumAnimKeys = Animation->GetDataModel()->GetNumberOfKeys();
+	for (int32 AnimKey = 0; AnimKey < NumAnimKeys; AnimKey++)
+	{
+		const FInt32Range KeyRangeToSet(AnimKey, AnimKey + 1);
+
+		const FTransform RootTransformOriginal = Extractor.GetGlobalBoneTransform(AnimKey, 0);
+		FTransform RootTransformNew = RootTransformOriginal;
+		if (TotalWeightToComputeRootPosition > UE_KINDA_SMALL_NUMBER)
+		{
+			FVector WeightedBoneTranslation = FVector::ZeroVector;
+			for (const FEncodeRootBoneWeightedBone& WeightedBone : WeightedBoneToComputeRootPosition)
+			{
+				WeightedBoneTranslation += Extractor.GetGlobalBoneTransform(AnimKey, WeightedBone.Bone.BoneName).GetTranslation() * WeightedBone.Weight;
+			}
+
+			WeightedBoneTranslation /= TotalWeightToComputeRootPosition;
+			WeightedBoneTranslation.Z = RootTransformNew.GetTranslation().Z;
+			RootTransformNew.SetTranslation(WeightedBoneTranslation);
+		}
+
+		if (TotalWeightToComputeRootOrientation > UE_KINDA_SMALL_NUMBER)
+		{
+			FVector WeightedBoneHeading = FVector::ZeroVector;
+			for (const FEncodeRootBoneWeightedBoneAxis& WeightedBoneAxis : WeightedBoneToComputeRootOrientation)
+			{
+				const FQuat GlobalBoneRotation = Extractor.GetGlobalBoneTransform(AnimKey, WeightedBoneAxis.Bone.BoneName).GetRotation();
+				switch (WeightedBoneAxis.BoneAxis)
+				{
+				case EEncodeRootBoneAxis::X:
+					WeightedBoneHeading += GlobalBoneRotation.GetAxisX() * WeightedBoneAxis.Weight;
+					break;
+				case EEncodeRootBoneAxis::Y:
+					WeightedBoneHeading += GlobalBoneRotation.GetAxisY() * WeightedBoneAxis.Weight;
+					break;
+				case EEncodeRootBoneAxis::Z:
+					WeightedBoneHeading += GlobalBoneRotation.GetAxisZ() * WeightedBoneAxis.Weight;
+					break;
+				}
+			}
+
+			const float Yaw = FMath::Atan2(-WeightedBoneHeading.X, WeightedBoneHeading.Y) * (180.f / UE_PI);
+			const FRotator Rotator(0.f, Yaw, 0.f);
+			RootTransformNew.SetRotation(Rotator.Quaternion());
+		}
+
+		Animation->GetController().UpdateBoneTrackKeys(Animation->GetSkeleton()->GetReferenceSkeleton().GetBoneName(0), KeyRangeToSet, { RootTransformNew.GetLocation() }, { RootTransformNew.GetRotation() }, { RootTransformNew.GetScale3D() });
+
+		// Now the mesh is facing the wrong axis. Update direct children of the root with the local space transform that puts them back to where they were originally
+		for (const FName& ChildBoneName : RootBoneChildBones)
+		{
+			const FTransform ChildBoneNewTransform = Extractor.GetGlobalBoneTransform(AnimKey, ChildBoneName).GetRelativeTransform(RootTransformNew);
+			Animation->GetController().UpdateBoneTrackKeys(ChildBoneName, KeyRangeToSet, { ChildBoneNewTransform.GetLocation() }, { ChildBoneNewTransform.GetRotation() }, { ChildBoneNewTransform.GetScale3D() });
+		}
+	}
+
+	// Done editing animation data
+	Animation->GetController().CloseBracket(bShouldTransact);
+
+
+#if ENABLE_ENCODE_ROOT_BONE_MODIFIER_TEST && DO_CHECK
+	// testing
+	FPoseExtractorUtilityLibrary TestExtractor(Animation);
+	for (const FEncodeRootBoneWeightedBone& WeightedBone : WeightedBoneToComputeRootPosition)
+	{
+		check(TestExtractor.AddBone(WeightedBone.Bone));
+	}
+	for (const FEncodeRootBoneWeightedBone& WeightedBone : WeightedBoneToComputeRootOrientation)
+	{
+		check(TestExtractor.AddBone(WeightedBone.Bone));
+	}
+	for (const int32 ChildBoneIndex : ChildBoneIndices)
+	{
+		const FName ChildBoneName = Animation->GetSkeleton()->GetReferenceSkeleton().GetBoneName(ChildBoneIndex);
+		if (Animation->GetDataModel()->IsValidBoneTrackName(ChildBoneName))
+		{
+			check(TestExtractor.AddBone(ChildBoneName, ChildBoneIndex));
+		}
+	}
+
+	check(TestExtractor.CacheGlobalBoneTransforms());
+
+	for (const FName& ChildBoneName : RootBoneChildBones)
+	{
+		for (int32 AnimKey = 0; AnimKey < NumAnimKeys; AnimKey++)
+		{
+			const FTransform ChildBoneTransform = TestExtractor.GetGlobalBoneTransform(AnimKey, ChildBoneName);
+			const FTransform TestChildBoneTransform = TestExtractor.GetGlobalBoneTransform(AnimKey, ChildBoneName);
+
+			static float EncodeRootBoneModifierTolerance = UE_KINDA_SMALL_NUMBER;
+			check(ChildBoneTransform.Equals(TestChildBoneTransform, EncodeRootBoneModifierTolerance));
+		}
+	}
+
+#endif // ENABLE_ENCODE_ROOT_BONE_MODIFIER_TEST && DO_CHECK
+}
+
+#undef LOCTEXT_NAMESPACE
+
+#endif // WITH_EDITOR
diff --git a/Engine/Plugins/Animation/AnimationModifierLibrary/Source/AnimationModifierLibrary/Private/PoseExtractorUtilityLibrary.cpp b/Engine/Plugins/Animation/AnimationModifierLibrary/Source/AnimationModifierLibrary/Private/PoseExtractorUtilityLibrary.cpp
new file mode 100644
--- /dev/null
+++ b/Engine/Plugins/Animation/AnimationModifierLibrary/Source/AnimationModifierLibrary/Private/PoseExtractorUtilityLibrary.cpp
@@ -0,0 +1,149 @@
+// Copyright Epic Games, Inc. All Rights Reserved.
+
+#if WITH_EDITOR
+
+#include "PoseExtractorUtilityLibrary.h"
+#include "Animation/AnimSequence.h"
+#include "Animation/Skeleton.h"
+#include "AnimationRuntime.h"
+
+FPoseExtractorUtilityLibrary::FPoseExtractorUtilityLibrary(UAnimSequence* Animation)
+	: Skeleton(nullptr)
+	, Model(nullptr)
+{
+	if (!Animation)
+	{
+		UE_LOG(LogAnimation, Error, TEXT("FPoseExtractorUtilityLibrary failed. Reason: Invalid Animation"));
+		return;
+	}
+
+	IAnimationDataController& Controller = Animation->GetController();
+	const IAnimationDataModel* AnimationModel = Animation->GetDataModel();
+
+	if (!AnimationModel)
+	{
+		UE_LOG(LogAnimation, Error, TEXT("FPoseExtractorUtilityLibrary failed. Reason: Invalid Data Model. Animation: %s"), *GetNameSafe(Animation));
+		return;
+	}
+
+	USkeleton* AnimationSkeleton = Animation->GetSkeleton();
+	if (!AnimationSkeleton)
+	{
+		UE_LOG(LogAnimation, Error, TEXT("FPoseExtractorUtilityLibrary failed. Reason: Invalid Skeleton. Animation: %s"), *GetNameSafe(Animation));
+		return;
+	}
+
+	const FReferenceSkeleton& RefSkeleton = AnimationSkeleton->GetReferenceSkeleton();
+	if (RefSkeleton.GetNum() == 0)
+	{
+		UE_LOG(LogAnimation, Error, TEXT("FPoseExtractorUtilityLibrary failed. Reason: Ref Skeleton. Animation: %s"), *GetNameSafe(Animation));
+		return;
+	}
+
+	Skeleton = AnimationSkeleton;
+	Model = AnimationModel;
+}
+
+bool FPoseExtractorUtilityLibrary::AddBone(const FBoneReference& Bone)
+{
+	if (!Skeleton)
+	{
+		return false;
+	}
+
+	check(ReducedIndexMap.IsEmpty());
+
+	FBoneReference TempBone = Bone;
+	TempBone.Initialize(Skeleton);
+	if (!TempBone.HasValidSetup())
+	{
+		UE_LOG(LogAnimation, Error, TEXT("FPoseExtractorUtilityLibrary::AddBone failed. Reason: Bone %s has invalid setup"), *TempBone.BoneName.ToString());
+		return false;
+	}
+
+	return AddBone(TempBone.BoneName, TempBone.BoneIndex);
+}
+
+bool FPoseExtractorUtilityLibrary::AddBone(const FName& BoneName, int32 BoneIndex)
+{
+	if (BoneIndex >= 0 && BoneIndex < MAX_uint16)
+	{
+		NameToBoneIndex.FindOrAdd(BoneName) = BoneIndex;
+		BoneIndicesWithParents.AddUnique(BoneIndex);
+		return true;
+	}
+	
+	UE_LOG(LogAnimation, Error, TEXT("FPoseExtractorUtilityLibrary::AddBone failed. Reason: Bone %s has invalid BoneIndex %d"), *BoneName.ToString(), BoneIndex);
+	return false;
+}
+
+bool FPoseExtractorUtilityLibrary::CacheGlobalBoneTransforms()
+{
+	if (!Skeleton)
+	{
+		return false;
+	}
+
+	if (BoneIndicesWithParents.IsEmpty())
+	{
+		UE_LOG(LogAnimation, Warning, TEXT("FPoseExtractorUtilityLibrary::CacheGlobalBoneTransforms no weighted bones?"));
+		return false;
+	}
+
+	BoneIndicesWithParents.Sort();
+
+	const FReferenceSkeleton& RefSkeleton = Skeleton->GetReferenceSkeleton();
+	check(RefSkeleton.GetNum() > 0);
+
+	FAnimationRuntime::EnsureParentsPresent(BoneIndicesWithParents, RefSkeleton);
+
+	for (int i = 0; i < BoneIndicesWithParents.Num(); ++i)
+	{
+		ReducedIndexMap.Add(BoneIndicesWithParents[i]) = i;
+	}
+	
+	const int32 NumAnimKeys = Model->GetNumberOfKeys();
+	GlobalBoneTransforms.SetNum(NumAnimKeys);
+	for (int32 AnimKey = 0; AnimKey < NumAnimKeys; AnimKey++)
+	{
+		GlobalBoneTransforms[AnimKey].SetNum(BoneIndicesWithParents.Num());
+
+		// calculating all the required global transforms
+		for (int i = 0; i < BoneIndicesWithParents.Num(); ++i)
+		{
+			const uint16 BoneIndex = BoneIndicesWithParents[i];
+			const FName BoneName = RefSkeleton.GetBoneName(BoneIndex);
+			const FTransform LocalBoneTransform = Model->EvaluateBoneTrackTransform(BoneName, AnimKey, EAnimInterpolationType::Step);
+			const int32 ParentBoneIndex = RefSkeleton.GetParentIndex(BoneIndex);
+			if (ParentBoneIndex != INDEX_NONE)
+			{
+				const uint16* MappedParentBoneIndex = ReducedIndexMap.Find(ParentBoneIndex);
+				check(MappedParentBoneIndex);
+				check(*MappedParentBoneIndex <= i);
+				GlobalBoneTransforms[AnimKey][i] = LocalBoneTransform * GlobalBoneTransforms[AnimKey][*MappedParentBoneIndex];
+			}
+			else
+			{
+				GlobalBoneTransforms[AnimKey][i] = LocalBoneTransform;
+			}
+		}
+	}
+
+	return true;
+}
+
+const FTransform& FPoseExtractorUtilityLibrary::GetGlobalBoneTransform(int32 AnimKey, uint16 BoneIndex) const
+{
+	const uint16* MappedParentBoneIndex = ReducedIndexMap.Find(BoneIndex);
+	check(MappedParentBoneIndex);
+	return GlobalBoneTransforms[AnimKey][*MappedParentBoneIndex];
+}
+
+const FTransform& FPoseExtractorUtilityLibrary::GetGlobalBoneTransform(int32 AnimKey, const FName& BoneName) const
+{
+	const uint16* BoneIndex = NameToBoneIndex.Find(BoneName);
+	check(BoneIndex);
+	return GetGlobalBoneTransform(AnimKey, *BoneIndex);
+}
+
+#endif // WITH_EDITOR
diff --git a/Engine/Plugins/Animation/AnimationModifierLibrary/Source/AnimationModifierLibrary/Private/PoseExtractorUtilityLibrary.h b/Engine/Plugins/Animation/AnimationModifierLibrary/Source/AnimationModifierLibrary/Private/PoseExtractorUtilityLibrary.h
new file mode 100644
--- /dev/null
+++ b/Engine/Plugins/Animation/AnimationModifierLibrary/Source/AnimationModifierLibrary/Private/PoseExtractorUtilityLibrary.h
@@ -0,0 +1,32 @@
+// Copyright Epic Games, Inc. All Rights Reserved.
+
+#pragma once
+
+#if WITH_EDITOR
+
+#include "CoreMinimal.h"
+
+struct FBoneReference;
+class IAnimationDataModel;
+class UAnimSequence;
+class USkeleton;
+
+struct FPoseExtractorUtilityLibrary
+{
+	FPoseExtractorUtilityLibrary(UAnimSequence* Animation);
+	bool AddBone(const FBoneReference& Bone);
+	bool AddBone(const FName& BoneName, int32 BoneIndex);
+	bool CacheGlobalBoneTransforms();
+	const FTransform& GetGlobalBoneTransform(int32 AnimKey, uint16 BoneIndex) const;
+	const FTransform& GetGlobalBoneTransform(int32 AnimKey, const FName& BoneName) const;
+	
+private:
+	const USkeleton* Skeleton = nullptr;
+	const IAnimationDataModel* Model = nullptr;
+	TArray<uint16> BoneIndicesWithParents;
+	TMap<uint16, uint16> ReducedIndexMap;
+	TMap<FName, uint16> NameToBoneIndex;
+	TArray<TArray<FTransform>> GlobalBoneTransforms;
+};
+
+#endif // WITH_EDITOR
diff --git a/Engine/Plugins/Animation/AnimationModifierLibrary/Source/AnimationModifierLibrary/Public/EncodeRootBoneModifier.h b/Engine/Plugins/Animation/AnimationModifierLibrary/Source/AnimationModifierLibrary/Public/EncodeRootBoneModifier.h
new file mode 100644
--- /dev/null
+++ b/Engine/Plugins/Animation/AnimationModifierLibrary/Source/AnimationModifierLibrary/Public/EncodeRootBoneModifier.h
@@ -0,0 +1,53 @@
+// Copyright Epic Games, Inc. All Rights Reserved.
+
+#pragma once
+
+#if WITH_EDITOR
+#include "AnimationModifier.h"
+#include "EncodeRootBoneModifier.generated.h"
+
+UENUM()
+enum class EEncodeRootBoneAxis : uint8
+{
+	X,
+	Y,
+	Z,
+};
+
+USTRUCT()
+struct FEncodeRootBoneWeightedBone
+{
+	GENERATED_BODY()
+
+	UPROPERTY(EditAnywhere, Category = "Settings")
+	FBoneReference Bone;
+
+	UPROPERTY(EditAnywhere, Category = "Settings")
+	float Weight = 1.f;
+};
+
+USTRUCT()
+struct FEncodeRootBoneWeightedBoneAxis : public FEncodeRootBoneWeightedBone
+{
+	GENERATED_BODY()
+
+	UPROPERTY(EditAnywhere, Category = "Settings")
+	EEncodeRootBoneAxis BoneAxis = EEncodeRootBoneAxis::Y;
+};
+
+UCLASS()
+class UEncodeRootBoneModifier : public UAnimationModifier
+{
+	GENERATED_BODY()
+
+public:
+	UPROPERTY(EditAnywhere, Category = "Settings")
+	TArray<FEncodeRootBoneWeightedBone> WeightedBoneToComputeRootPosition;
+
+	UPROPERTY(EditAnywhere, Category = "Settings")
+	TArray<FEncodeRootBoneWeightedBoneAxis> WeightedBoneToComputeRootOrientation;
+
+	virtual void OnApply_Implementation(UAnimSequence* Animation) override;
+};
+
+#endif // WITH_EDITOR
\ No newline at end of file
