Truncation fixes for AnimationModifierLibrary

diff --git a/Engine/Plugins/Animation/AnimationModifierLibrary/Source/AnimationModifierLibrary/AnimationModifierLibrary.Build.cs b/Engine/Plugins/Animation/AnimationModifierLibrary/Source/AnimationModifierLibrary/AnimationModifierLibrary.Build.cs
--- a/Engine/Plugins/Animation/AnimationModifierLibrary/Source/AnimationModifierLibrary/AnimationModifierLibrary.Build.cs
+++ b/Engine/Plugins/Animation/AnimationModifierLibrary/Source/AnimationModifierLibrary/AnimationModifierLibrary.Build.cs
@@ -31,6 +31,7 @@ namespace UnrealBuildTool.Rules
 				}
 			);
 
+			UnsafeTypeCastWarningLevel = WarningLevel.Error;
 		}
 	}
 }
diff --git a/Engine/Plugins/Animation/AnimationModifierLibrary/Source/AnimationModifierLibrary/Private/EncodeRootBoneModifier.cpp b/Engine/Plugins/Animation/AnimationModifierLibrary/Source/AnimationModifierLibrary/Private/EncodeRootBoneModifier.cpp
--- a/Engine/Plugins/Animation/AnimationModifierLibrary/Source/AnimationModifierLibrary/Private/EncodeRootBoneModifier.cpp
+++ b/Engine/Plugins/Animation/AnimationModifierLibrary/Source/AnimationModifierLibrary/Private/EncodeRootBoneModifier.cpp
@@ -101,7 +101,7 @@ void UEncodeRootBoneModifier::OnApply_Implementation(UAnimSequence* Animation)
 				}
 			}
 
-			const float Yaw = FMath::Atan2(-WeightedBoneHeading.X, WeightedBoneHeading.Y) * (180.f / UE_PI);
+			const double Yaw = FMath::Atan2(-WeightedBoneHeading.X, WeightedBoneHeading.Y) * (180.f / UE_PI);
 			const FRotator Rotator(0.f, Yaw, 0.f);
 			RootTransformNew.SetRotation(Rotator.Quaternion());
 		}
diff --git a/Engine/Plugins/Animation/AnimationModifierLibrary/Source/AnimationModifierLibrary/Private/MotionExtractorModifier.cpp b/Engine/Plugins/Animation/AnimationModifierLibrary/Source/AnimationModifierLibrary/Private/MotionExtractorModifier.cpp
--- a/Engine/Plugins/Animation/AnimationModifierLibrary/Source/AnimationModifierLibrary/Private/MotionExtractorModifier.cpp
+++ b/Engine/Plugins/Animation/AnimationModifierLibrary/Source/AnimationModifierLibrary/Private/MotionExtractorModifier.cpp
@@ -62,7 +62,7 @@ void UMotionExtractorModifier::OnApply_Implementation(UAnimSequence* Animation)
 	TGuardValue<bool> ForceRootLockGuard(Animation->bForceRootLock, false);
 
 	TArray<FBoneIndexType> RequiredBones;
-	RequiredBones.Add(BoneIndex);
+	RequiredBones.Add(IntCastChecked<FBoneIndexType>(BoneIndex));
 	Skeleton->GetReferenceSkeleton().EnsureParentsExistAndSort(RequiredBones);
 
 	FBoneContainer BoneContainer(RequiredBones, UE::Anim::ECurveFilterMode::DisallowAll, *Skeleton);
@@ -71,7 +71,7 @@ void UMotionExtractorModifier::OnApply_Implementation(UAnimSequence* Animation)
 	FTransform FirstFrameBoneTransform = UMotionExtractorUtilityLibrary::ExtractBoneTransform(Animation, BoneContainer, CompactPoseBoneIndex, 0.f, bComponentSpace);
 
 	const float AnimLength = Animation->GetPlayLength();
-	const float SampleInterval = 1.f / SampleRate;
+	const float SampleInterval = 1.f / static_cast<float>(SampleRate);
 
 	FTransform LastBoneTransform = FTransform::Identity;
 	float Time = 0.f;
@@ -82,7 +82,7 @@ void UMotionExtractorModifier::OnApply_Implementation(UAnimSequence* Animation)
 	{
 		while (Time < AnimLength)
 		{
-			Time = FMath::Clamp(SampleIndex * SampleInterval, 0.f, AnimLength);
+			Time = FMath::Clamp(static_cast<float>(SampleIndex) * SampleInterval, 0.f, AnimLength);
 			SampleIndex++;
 
 			FTransform BoneTransform = UMotionExtractorUtilityLibrary::ExtractBoneTransform(Animation, BoneContainer, CompactPoseBoneIndex, Time, bComponentSpace);
@@ -111,7 +111,7 @@ void UMotionExtractorModifier::OnApply_Implementation(UAnimSequence* Animation)
 	
 	while (Time < AnimLength)
 	{
-		Time = FMath::Clamp(SampleIndex * SampleInterval, 0.f, AnimLength);
+		Time = FMath::Clamp(static_cast<float>(SampleIndex) * SampleInterval, 0.f, AnimLength);
 		SampleIndex++;
 
 		FTransform BoneTransform = UMotionExtractorUtilityLibrary::ExtractBoneTransform(Animation, BoneContainer, CompactPoseBoneIndex, Time, bComponentSpace);
diff --git a/Engine/Plugins/Animation/AnimationModifierLibrary/Source/AnimationModifierLibrary/Private/MotionExtractorUtilities.cpp b/Engine/Plugins/Animation/AnimationModifierLibrary/Source/AnimationModifierLibrary/Private/MotionExtractorUtilities.cpp
--- a/Engine/Plugins/Animation/AnimationModifierLibrary/Source/AnimationModifierLibrary/Private/MotionExtractorUtilities.cpp
+++ b/Engine/Plugins/Animation/AnimationModifierLibrary/Source/AnimationModifierLibrary/Private/MotionExtractorUtilities.cpp
@@ -21,14 +21,14 @@ namespace UE::Anim::MotionExtractorUtility
 
 		const float AnimLength = Animation->GetPlayLength();
 		float SampleInterval = 1.f / SampleRate;
-		int32 NumSteps = AnimLength / SampleInterval;
+		int32 NumSteps = static_cast<int32>(AnimLength / SampleInterval);
 
 		TArray<FVector2D> OutRanges;
 		OutRanges.Reserve(4);
 		bool bMeetsCriterion = false;
 		for (int32 Step = 0; Step < NumSteps; ++Step)
 		{
-			const float Time = Step * SampleInterval;
+			const float Time = static_cast<float>(Step) * SampleInterval;
 			const bool bAllowLooping = false;
 			const FTransform RootMotion = Animation->ExtractRootMotion(Time, SampleInterval, bAllowLooping);
 			const bool bMetCriterion = bMeetsCriterion;
@@ -98,7 +98,7 @@ float UMotionExtractorUtilityLibrary::GetDesiredValue(
 		return 0.0f;
 	}
 
-	float Value = 0.f;
+	double Value = 0.f;
 	switch (MotionType)
 	{
 	case EMotionExtractor_MotionType::Translation:
@@ -141,7 +141,7 @@ float UMotionExtractorUtilityLibrary::GetDesiredValue(
 	{
 		if (!FMath::IsNearlyZero(DeltaTime))
 		{
-			const float Delta = UMotionExtractorUtilityLibrary::CalculateMagnitude((BoneTransform.GetTranslation() - LastBoneTransform.GetTranslation()), Axis);
+			const double Delta = UMotionExtractorUtilityLibrary::CalculateMagnitude((BoneTransform.GetTranslation() - LastBoneTransform.GetTranslation()), Axis);
 			Value = Delta / DeltaTime;
 		}
 
@@ -186,7 +186,7 @@ float UMotionExtractorUtilityLibrary::GetDesiredValue(
 	default: check(false); break;
 	}
 
-	return Value;
+	return static_cast<float>(Value);
 }
 
 TArray<FVector2D> UMotionExtractorUtilityLibrary::GetStoppedRangesFromRootMotion(const UAnimSequence* AnimSequence, float StopSpeedThreshold, float SampleRate)
@@ -194,7 +194,7 @@ TArray<FVector2D> UMotionExtractorUtilityLibrary::GetStoppedRangesFromRootMotion
 	auto Criterion = [StopSpeedThreshold](const FTransform& InRootMotionDelta, float DeltaTime)
 	{
 		const FVector RootMotionTranslation = InRootMotionDelta.GetTranslation();
-		const float RootMotionSpeed = RootMotionTranslation.Size() / DeltaTime;
+		const double RootMotionSpeed = RootMotionTranslation.Size() / DeltaTime;
 		return RootMotionSpeed < StopSpeedThreshold;
 	};
 
@@ -206,7 +206,7 @@ TArray<FVector2D> UMotionExtractorUtilityLibrary::GetMovingRangesFromRootMotion(
 	auto Criterion = [StopSpeedThreshold](const FTransform& InRootMotionDelta, float DeltaTime)
 	{
 		const FVector RootMotionTranslation = InRootMotionDelta.GetTranslation();
-		const float RootMotionSpeed = RootMotionTranslation.Size() / DeltaTime;
+		const double RootMotionSpeed = RootMotionTranslation.Size() / DeltaTime;
 		return RootMotionSpeed > StopSpeedThreshold;
 	};
 
@@ -241,7 +241,7 @@ FTransform UMotionExtractorUtilityLibrary::ExtractBoneTransform(UAnimSequence* A
 	}
 }
 
-float UMotionExtractorUtilityLibrary::CalculateMagnitude(const FVector& Vector, EMotionExtractor_Axis Axis)
+double UMotionExtractorUtilityLibrary::CalculateMagnitude(const FVector& Vector, EMotionExtractor_Axis Axis)
 {
 	switch (Axis)
 	{
diff --git a/Engine/Plugins/Animation/AnimationModifierLibrary/Source/AnimationModifierLibrary/Private/PoseExtractorUtilityLibrary.cpp b/Engine/Plugins/Animation/AnimationModifierLibrary/Source/AnimationModifierLibrary/Private/PoseExtractorUtilityLibrary.cpp
--- a/Engine/Plugins/Animation/AnimationModifierLibrary/Source/AnimationModifierLibrary/Private/PoseExtractorUtilityLibrary.cpp
+++ b/Engine/Plugins/Animation/AnimationModifierLibrary/Source/AnimationModifierLibrary/Private/PoseExtractorUtilityLibrary.cpp
@@ -68,8 +68,8 @@ bool FPoseExtractorUtilityLibrary::AddBone(const FName& BoneName, int32 BoneInde
 {
 	if (BoneIndex >= 0 && BoneIndex < MAX_uint16)
 	{
-		NameToBoneIndex.FindOrAdd(BoneName) = BoneIndex;
-		BoneIndicesWithParents.AddUnique(BoneIndex);
+		NameToBoneIndex.FindOrAdd(BoneName) = static_cast<FBoneIndexType>(BoneIndex);
+		BoneIndicesWithParents.AddUnique(static_cast<FBoneIndexType>(BoneIndex));
 		return true;
 	}
 	
@@ -97,9 +97,9 @@ bool FPoseExtractorUtilityLibrary::CacheGlobalBoneTransforms()
 
 	FAnimationRuntime::EnsureParentsPresent(BoneIndicesWithParents, RefSkeleton);
 
-	for (int i = 0; i < BoneIndicesWithParents.Num(); ++i)
+	for (int32 i = 0; i < BoneIndicesWithParents.Num(); ++i)
 	{
-		ReducedIndexMap.Add(BoneIndicesWithParents[i]) = i;
+		ReducedIndexMap.Add(BoneIndicesWithParents[i]) = IntCastChecked<FBoneIndexType>(i);
 	}
 	
 	const int32 NumAnimKeys = Model->GetNumberOfKeys();
@@ -117,7 +117,7 @@ bool FPoseExtractorUtilityLibrary::CacheGlobalBoneTransforms()
 			const int32 ParentBoneIndex = RefSkeleton.GetParentIndex(BoneIndex);
 			if (ParentBoneIndex != INDEX_NONE)
 			{
-				const uint16* MappedParentBoneIndex = ReducedIndexMap.Find(ParentBoneIndex);
+				const uint16* MappedParentBoneIndex = ReducedIndexMap.Find(IntCastChecked<FBoneIndexType>(ParentBoneIndex));
 				check(MappedParentBoneIndex);
 				check(*MappedParentBoneIndex <= i);
 				GlobalBoneTransforms[AnimKey][i] = LocalBoneTransform * GlobalBoneTransforms[AnimKey][*MappedParentBoneIndex];
diff --git a/Engine/Plugins/Animation/AnimationModifierLibrary/Source/AnimationModifierLibrary/Public/MotionExtractorUtilities.h b/Engine/Plugins/Animation/AnimationModifierLibrary/Source/AnimationModifierLibrary/Public/MotionExtractorUtilities.h
--- a/Engine/Plugins/Animation/AnimationModifierLibrary/Source/AnimationModifierLibrary/Public/MotionExtractorUtilities.h
+++ b/Engine/Plugins/Animation/AnimationModifierLibrary/Source/AnimationModifierLibrary/Public/MotionExtractorUtilities.h
@@ -76,7 +76,7 @@ public:
 	static FTransform ExtractBoneTransform(UAnimSequence* Animation, const FBoneContainer& BoneContainer, FCompactPoseBoneIndex CompactPoseBoneIndex, float Time, bool bComponentSpace);
 
 	/** Helper function to calculate the magnitude of a vector only considering a specific axis or axes */
-	static float CalculateMagnitude(const FVector& Vector, EMotionExtractor_Axis Axis);
+	static double CalculateMagnitude(const FVector& Vector, EMotionExtractor_Axis Axis);
 };
 
 #if UE_ENABLE_INCLUDE_ORDER_DEPRECATED_IN_5_2
