Expose Motion Extractor Utilities library
Make EMotionExtractor_MotionType bitmask friendly
#rnx

diff --git a/Engine/Plugins/Animation/AnimationModifierLibrary/Source/AnimationModifierLibrary/Public/MotionExtractorTypes.h b/Engine/Plugins/Animation/AnimationModifierLibrary/Source/AnimationModifierLibrary/Public/MotionExtractorTypes.h
--- a/Engine/Plugins/Animation/AnimationModifierLibrary/Source/AnimationModifierLibrary/Public/MotionExtractorTypes.h
+++ b/Engine/Plugins/Animation/AnimationModifierLibrary/Source/AnimationModifierLibrary/Public/MotionExtractorTypes.h
@@ -5,15 +5,24 @@
 #include "MotionExtractorTypes.generated.h"
 
 /** Type of motion to extract */
-UENUM(BlueprintType)
+UENUM(BlueprintType, meta = (Bitflags, UseEnumValuesAsMaskValuesInEditor = "true"))
 enum class EMotionExtractor_MotionType : uint8
 {
-	Translation,
-	Rotation,
-	Scale,
-	TranslationSpeed,
-	RotationSpeed
+	None				= 0 UMETA(Hidden),
+	Translation			= 1 << 0,
+	Rotation			= 1 << 1,
+	Scale				= 1 << 2,
+	TranslationSpeed	= 1 << 3,
+	RotationSpeed		= 1 << 4,
 };
+ENUM_CLASS_FLAGS(EMotionExtractor_MotionType);
+constexpr bool EnumHasAnyFlags(int32 Flags, EMotionExtractor_MotionType Contains) { return (Flags & int32(Contains)) != 0; }
+ENUM_RANGE_BY_VALUES(EMotionExtractor_MotionType, 
+	EMotionExtractor_MotionType::Translation, 
+	EMotionExtractor_MotionType::Rotation,
+	EMotionExtractor_MotionType::Scale,
+	EMotionExtractor_MotionType::TranslationSpeed,
+	EMotionExtractor_MotionType::RotationSpeed);
 
 /** Axis to get the final value from */
 UENUM(BlueprintType)
diff --git a/Engine/Plugins/Animation/AnimationModifierLibrary/Source/AnimationModifierLibrary/Public/MotionExtractorUtilities.h b/Engine/Plugins/Animation/AnimationModifierLibrary/Source/AnimationModifierLibrary/Public/MotionExtractorUtilities.h
--- a/Engine/Plugins/Animation/AnimationModifierLibrary/Source/AnimationModifierLibrary/Public/MotionExtractorUtilities.h
+++ b/Engine/Plugins/Animation/AnimationModifierLibrary/Source/AnimationModifierLibrary/Public/MotionExtractorUtilities.h
@@ -13,7 +13,7 @@ struct FBoneContainer;
 struct FCompactPoseBoneIndex;
 
 UCLASS(meta=(ScriptName="MotionExtractorUtilityLibrary"))
-class  UMotionExtractorUtilityLibrary : public UBlueprintFunctionLibrary
+class ANIMATIONMODIFIERLIBRARY_API UMotionExtractorUtilityLibrary : public UBlueprintFunctionLibrary
 {
 	GENERATED_BODY()
 
