Fix for non-unity compile error.
#rnx

diff --git a/Engine/Plugins/Animation/AnimationModifierLibrary/Source/AnimationModifierLibrary/Public/MotionExtractorTypes.h b/Engine/Plugins/Animation/AnimationModifierLibrary/Source/AnimationModifierLibrary/Public/MotionExtractorTypes.h
--- a/Engine/Plugins/Animation/AnimationModifierLibrary/Source/AnimationModifierLibrary/Public/MotionExtractorTypes.h
+++ b/Engine/Plugins/Animation/AnimationModifierLibrary/Source/AnimationModifierLibrary/Public/MotionExtractorTypes.h
@@ -2,6 +2,8 @@
 
 #pragma once
 
+#include "Misc/EnumClassFlags.h"
+#include "Misc/EnumRange.h"
 #include "MotionExtractorTypes.generated.h"
 
 /** Type of motion to extract */
