Add animation modifier to help fixup root motion data recorded via Take Recorder.

diff --git a/Engine/Plugins/Animation/AnimationModifierLibrary/Source/AnimationModifierLibrary/Private/ZeroOutRootBoneModifier.cpp b/Engine/Plugins/Animation/AnimationModifierLibrary/Source/AnimationModifierLibrary/Private/ZeroOutRootBoneModifier.cpp
new file mode 100644
--- /dev/null
+++ b/Engine/Plugins/Animation/AnimationModifierLibrary/Source/AnimationModifierLibrary/Private/ZeroOutRootBoneModifier.cpp
@@ -0,0 +1,113 @@
+// Copyright Epic Games, Inc. All Rights Reserved.
+
+#include "ZeroOutRootBoneModifier.h"
+#include "Animation/AnimData/AnimDataModel.h"
+#include "Animation/AnimSequence.h"
+#include "Animation/AnimData/IAnimationDataController.h"
+#include "Animation/AnimSequenceHelpers.h"
+#include "AnimationUtils.h"
+#include "Animation/Skeleton.h"
+#include "EngineLogs.h"
+
+#define LOCTEXT_NAMESPACE "ZeroOutRootBoneModifier"
+
+UZeroOutRootBoneModifier::UZeroOutRootBoneModifier()
+	:Super()
+{
+}
+
+void UZeroOutRootBoneModifier::OnApply_Implementation(UAnimSequence* Animation)
+{
+	if (Animation == nullptr)
+	{
+		UE_LOG(LogAnimation, Error, TEXT("ZeroOutRootBoneModifier failed. Reason: Invalid Animation"));
+		return;
+	}
+
+	IAnimationDataController& Controller = Animation->GetController();
+	const IAnimationDataModel* Model = Animation->GetDataModel();
+
+	if (Model == nullptr)
+	{
+		UE_LOG(LogAnimation, Error, TEXT("ZeroOutRootBoneModifier failed. Reason: Invalid Data Model. Animation: %s"), *GetNameSafe(Animation));
+		return;
+	}
+
+	const USkeleton* Skeleton = Animation->GetSkeleton();
+	if (Skeleton == nullptr)
+	{
+		UE_LOG(LogAnimation, Error, TEXT("ZeroOutRootBoneModifier failed. Reason: Invalid Skeleton. Animation: %s"), *GetNameSafe(Animation));
+		return;
+	}
+
+	const FReferenceSkeleton& RefSkeleton = Skeleton->GetReferenceSkeleton();
+	if (RefSkeleton.GetNum() == 0)
+	{
+		UE_LOG(LogAnimation, Error, TEXT("ZeroOutRootBoneModifier failed. Reason: Ref Skeleton. Animation: %s"), *GetNameSafe(Animation));
+		return;
+	}
+
+	const FName RootBoneName = RefSkeleton.GetBoneName(0);
+
+	// Start editing animation data
+	const bool bShouldTransact = false;
+	Controller.OpenBracket(LOCTEXT("ZeroOutRootBoneModifier_Bracket", "Zero Out Root Bone"), bShouldTransact);
+
+	const int32 NumKeys = Model->GetNumberOfKeys();
+
+	// Cache root transforms and find where motion starts/ends.
+	TArray<FTransform> PreModifiedRootTransforms;
+	PreModifiedRootTransforms.Reserve(NumKeys);
+	int32 KeyWhereMotionStarts = -1;
+	int32 KeyWhereMotionEnds = 0;
+	for (int32 AnimKey = 0; AnimKey < NumKeys; AnimKey++)
+	{
+		const FTransform RootTransform = Model->EvaluateBoneTrackTransform(RootBoneName, AnimKey, EAnimInterpolationType::Step);
+		PreModifiedRootTransforms.Add(RootTransform);
+
+		if (AnimKey + 1 < NumKeys)
+		{
+			const FTransform NextRootTransform = Model->EvaluateBoneTrackTransform(RootBoneName, AnimKey + 1, EAnimInterpolationType::Step);
+			const bool bRootTransformChanged = !FTransform::AreTranslationsEqual(RootTransform, NextRootTransform) || !FTransform::AreRotationsEqual(RootTransform, NextRootTransform);
+
+			if (bRootTransformChanged)
+			{
+				if (KeyWhereMotionStarts < 0)
+				{
+					KeyWhereMotionStarts = AnimKey;
+				}
+
+				KeyWhereMotionEnds = AnimKey + 1;
+			}
+		}
+	}
+
+	// For each key in the animation
+	for (int32 AnimKey = 0; AnimKey < NumKeys; AnimKey++)
+	{
+		const FInt32Range KeyRangeToSet(AnimKey, AnimKey + 1);
+
+		FTransform NewRootTransform = PreModifiedRootTransforms[AnimKey] * PreModifiedRootTransforms[0].Inverse();
+		Controller.UpdateBoneTrackKeys(RootBoneName, KeyRangeToSet, { NewRootTransform.GetLocation() }, { NewRootTransform.GetRotation() }, { NewRootTransform.GetScale3D() });
+	}
+
+	// Done editing animation data
+	Controller.CloseBracket(bShouldTransact);
+	
+	if (bClipEndFramesWithNoMotion && (KeyWhereMotionEnds > KeyWhereMotionStarts) && (KeyWhereMotionEnds + 1 < NumKeys))
+	{
+		// Remove everything after the frame where motion ends.
+		const int32 FirstKeyToRemove = KeyWhereMotionEnds + 1;
+		const int32 NumKeysToRemove = NumKeys - FirstKeyToRemove;
+		UE::Anim::AnimationData::RemoveKeys(Animation, FirstKeyToRemove, NumKeysToRemove);
+	}
+
+	if (bClipStartFramesWithNoMotion && (KeyWhereMotionStarts > 1))
+	{
+		// Remove everything before the frame where motion starts.
+		const int32 NumKeysToRemove = KeyWhereMotionStarts - 1;
+		UE::Anim::AnimationData::RemoveKeys(Animation, 0, NumKeysToRemove);
+	}
+}
+
+#undef LOCTEXT_NAMESPACE
\ No newline at end of file
diff --git a/Engine/Plugins/Animation/AnimationModifierLibrary/Source/AnimationModifierLibrary/Public/ZeroOutRootBoneModifier.h b/Engine/Plugins/Animation/AnimationModifierLibrary/Source/AnimationModifierLibrary/Public/ZeroOutRootBoneModifier.h
new file mode 100644
--- /dev/null
+++ b/Engine/Plugins/Animation/AnimationModifierLibrary/Source/AnimationModifierLibrary/Public/ZeroOutRootBoneModifier.h
@@ -0,0 +1,29 @@
+// Copyright Epic Games, Inc. All Rights Reserved.
+
+#pragma once
+
+#include "AnimationModifier.h"
+#include "ZeroOutRootBoneModifier.generated.h"
+
+/** Adjust root motion to be relative to the first frame. Optionally cut the start and end sections of the animation that don't have motion on the root.
+	This was written to be used when capturing Character Movement motion via Take Recorder. Take Recorder outputs an animation captured from
+	a character moving in game in world space, and this modifier zeroes out the root. The animation can then be exported to fbx to be animated against. */
+UCLASS()
+class UZeroOutRootBoneModifier : public UAnimationModifier
+{
+	GENERATED_BODY()
+
+public:
+
+	UZeroOutRootBoneModifier();
+
+	virtual void OnApply_Implementation(UAnimSequence* Animation) override;
+
+	// Clip frames at the start of the animation that have no root motion.
+	UPROPERTY(EditAnywhere, BlueprintReadOnly, Category = "Settings")
+	bool bClipStartFramesWithNoMotion = true;
+
+	// Clip frames at the end of the animation that have no root motion.
+	UPROPERTY(EditAnywhere, BlueprintReadOnly, Category = "Settings")
+	bool bClipEndFramesWithNoMotion = true;
+};
\ No newline at end of file
