Added FootstepAnimEventsModifier to the AnimationModifierLibrary. Used to generate sync markers and notifies for any number of specified bones.

[RN] minor

diff --git a/Engine/Plugins/Animation/AnimationModifierLibrary/Source/AnimationModifierLibrary/Private/FootstepAnimEventsModifier.cpp b/Engine/Plugins/Animation/AnimationModifierLibrary/Source/AnimationModifierLibrary/Private/FootstepAnimEventsModifier.cpp
new file mode 100644
--- /dev/null
+++ b/Engine/Plugins/Animation/AnimationModifierLibrary/Source/AnimationModifierLibrary/Private/FootstepAnimEventsModifier.cpp
@@ -0,0 +1,265 @@
+﻿// Copyright Epic Games, Inc. All Rights Reserved.
+
+
+#if WITH_EDITOR
+
+#include "FootstepAnimEventsModifier.h"
+#include "Animation/AnimSequence.h"
+#include "AnimationBlueprintLibrary.h"
+#include "AnimPose.h"
+#include "Animation/AnimNotifies/AnimNotify.h"
+
+#define LOCTEXT_NAMESPACE "FootstepAnimEventsModifierBase"
+
+UFootstepAnimEventsModifier::UFootstepAnimEventsModifier() : Super()
+{
+	SampleStep = 0.0125f;
+	GroundThreshold = 0.95f;
+	bShouldRemovePreExistingNotifiesOrSyncMarkers = false;
+}
+
+void UFootstepAnimEventsModifier::OnApply_Implementation(UAnimSequence* InAnimation)
+{
+	Super::OnApply_Implementation(InAnimation);
+
+	if (InAnimation == nullptr)
+	{
+		UE_LOG(LogAnimation, Error, TEXT("FootstepAnimEventsModifierBase failed. Reason: Invalid Animation"));
+		return;
+	}
+	
+	// Disable root motion lock during application
+	bIsRootMotionLocked = InAnimation->bForceRootLock;
+	InAnimation->bForceRootLock = false;
+
+	// Clean up or generate tracks if needed
+	ValidateNotifyTracks(InAnimation);
+	
+	// Process animation
+	{
+		TArray<FFootSampleState> FootSampleStates;
+		FootSampleStates.Init(FFootSampleState(), FootDefinitions.Num());
+		
+		const float SequenceLength = InAnimation->GetPlayLength();
+		const int SampleNum = FMath::TruncToInt(SequenceLength / SampleStep);
+
+		// Get ground levels for each foot definition
+		for (int i = 0; i < SampleNum; ++i)
+		{
+			const float SampleTime = static_cast<float>(i) * SampleStep;
+			FAnimPose AnimPose;
+		
+			// Query animation pose at the current sample time
+			{
+				FAnimPoseEvaluationOptions AnimPoseEvalOptions;
+				AnimPoseEvalOptions.EvaluationType = EAnimDataEvalType::Raw;
+			
+				UAnimPoseExtensions::GetAnimPoseAtTime(InAnimation, SampleTime, AnimPoseEvalOptions, AnimPose);
+			}
+
+			// Find out ground level for each foot
+			for (int CurrFootIdx = 0; CurrFootIdx < FootDefinitions.Num(); ++CurrFootIdx)
+			{
+				const FFootDefinition & FootDef = FootDefinitions[CurrFootIdx];
+				FFootSampleState & FootState = FootSampleStates[CurrFootIdx];
+
+				// Compute ground level
+				FootState.GroundLevel = FMath::Min(FootState.GroundLevel, UAnimPoseExtensions::GetBonePose(AnimPose, FootDef.FootBoneName, EAnimPoseSpaces::World).GetLocation().Z);
+			}
+		}
+
+		// Avoid creating a notify at the very start of a sequence if the foot bone is already within ground threshold 
+		{
+			FAnimPose AnimPose;
+			
+			FAnimPoseEvaluationOptions AnimPoseEvalOptions;
+			AnimPoseEvalOptions.EvaluationType = EAnimDataEvalType::Raw;
+			
+			UAnimPoseExtensions::GetAnimPoseAtTime(InAnimation, 0.0f, AnimPoseEvalOptions, AnimPose);
+			
+			for (int CurrFootIdx = 0; CurrFootIdx < FootDefinitions.Num(); ++CurrFootIdx)
+			{
+				const FFootDefinition & FootDef = FootDefinitions[CurrFootIdx];
+				FFootSampleState & FootState = FootSampleStates[CurrFootIdx];
+				
+				FootState.bWasFootBoneInGround = FMath::Abs(FootState.GroundLevel - UAnimPoseExtensions::GetBonePose(AnimPose, FootDef.FootBoneName, EAnimPoseSpaces::World).GetLocation().Z) < GroundThreshold;
+			}
+		}
+		
+		// Generate animation events
+		for (int i = 0; i < SampleNum; ++i)
+		{
+			const float SampleTime = static_cast<float>(i) * SampleStep;
+
+			FAnimPose AnimPose;
+			FAnimPose FutureAnimPose;
+			
+			// Query animation pose at the current sample time
+			{
+				FAnimPoseEvaluationOptions AnimPoseEvalOptions;
+				AnimPoseEvalOptions.EvaluationType = EAnimDataEvalType::Raw;
+				
+				UAnimPoseExtensions::GetAnimPoseAtTime(InAnimation, SampleTime, AnimPoseEvalOptions, AnimPose);
+				UAnimPoseExtensions::GetAnimPoseAtTime(InAnimation, FMath::Clamp(SampleTime + SampleStep, 0.0f, SequenceLength), AnimPoseEvalOptions, FutureAnimPose);
+			}
+
+			// Process all foot definitions
+			for (int CurrFootIdx = 0; CurrFootIdx < FootDefinitions.Num(); ++CurrFootIdx)
+			{
+				const FFootDefinition & FootDef = FootDefinitions[CurrFootIdx];
+				FFootSampleState & FootState = FootSampleStates[CurrFootIdx];
+				
+				// Update current sample state
+				{
+					// Get reference bone translation
+					const FTransform ReferenceBoneTransform = UAnimPoseExtensions::GetBonePose(AnimPose, FootDef.ReferenceBoneName, EAnimPoseSpaces::World);
+					const FTransform FutureReferenceBoneTransform = UAnimPoseExtensions::GetBonePose(FutureAnimPose, FootDef.ReferenceBoneName, EAnimPoseSpaces::World);
+					const FVector ReferenceBoneTranslation = FutureReferenceBoneTransform.GetLocation() - ReferenceBoneTransform.GetLocation();
+
+					// Get foot bone direction with respect to reference bone
+					const FTransform FootBoneTransform = UAnimPoseExtensions::GetBonePose(AnimPose, FootDef.FootBoneName, EAnimPoseSpaces::World);
+					const FVector ReferenceBoneToFootBoneVector = FootBoneTransform.GetLocation() - ReferenceBoneTransform.GetLocation();
+
+					FootState.RefBoneTranslationDotRefBoneToFootBoneVec = FVector::DotProduct(ReferenceBoneTranslation, ReferenceBoneToFootBoneVector);
+					FootState.bIsFootBoneInGround = FMath::Abs(FootState.GroundLevel - FootBoneTransform.GetLocation().Z) < GroundThreshold;
+				}
+				
+				// Generate sync markers
+				if (FootDef.bShouldGenerateSyncMarkers && CanWePlaceEventAtSample(FootState, FootDef.SyncMarkerDetectionTechnique))
+				{
+					UAnimationBlueprintLibrary::AddAnimationSyncMarker(InAnimation, FootDef.SyncMarkerName, SampleTime, FootDef.SyncMarkerTrackName);
+				}
+
+				// Generate foot step fx notifies
+				if (FootDef.bShouldGenerateNotifies && CanWePlaceEventAtSample(FootState, FootDef.FootstepNotifyDetectionTechnique))
+				{
+					UAnimationBlueprintLibrary::AddAnimationNotifyEvent(InAnimation, FootDef.FootstepNotifyTrackName, SampleTime, FootDef.FootstepNotify);
+				}					
+
+				// Update foot state
+				FootState.bWasFootBoneInGround = FootState.bIsFootBoneInGround;
+				FootState.PrevRefBoneTranslationDotRefBoneToFootBoneVec = FootState.RefBoneTranslationDotRefBoneToFootBoneVec;
+			}
+		}
+	}
+
+	// Revert modifications to animation state
+	InAnimation->bForceRootLock = bIsRootMotionLocked;
+}
+
+void UFootstepAnimEventsModifier::OnRevert_Implementation(UAnimSequence* InAnimation)
+{
+	Super::OnRevert_Implementation(InAnimation);
+
+	// Delete any generate tracks.
+	for (const FName GeneratedNotifyTrack : GeneratedNotifyTracks)
+	{
+		UAnimationBlueprintLibrary::RemoveAnimationNotifyTrack(InAnimation, GeneratedNotifyTrack);
+	}
+	
+	// Delete any generated anim events.
+	for (const FName ProcessedNotifyTrack : ProcessedNotifyTracks)
+	{
+		const bool IsNotifyTrackValid = UAnimationBlueprintLibrary::IsValidAnimNotifyTrackName(InAnimation, ProcessedNotifyTrack);
+
+		if (IsNotifyTrackValid)
+		{
+			for (const FFootDefinition& FootDef : FootDefinitions)
+			{
+				if (FootDef.bShouldGenerateSyncMarkers && FootDef.SyncMarkerTrackName == ProcessedNotifyTrack)
+				{
+					UAnimationBlueprintLibrary::RemoveAnimationSyncMarkersByName(InAnimation, FootDef.SyncMarkerName);
+				}
+
+				if (FootDef.bShouldGenerateNotifies && FootDef.FootstepNotifyTrackName == ProcessedNotifyTrack)
+				{
+					if (FootDef.FootstepNotify)
+					{
+						UAnimationBlueprintLibrary::RemoveAnimationNotifyEventsByName(InAnimation, FName(Cast<UAnimNotify>(FootDef.FootstepNotify->GetDefaultObject())->GetNotifyName()));
+					}
+					else
+					{
+						UE_LOG(LogAnimation, Error, TEXT("FootstepAnimEventsModifierBase failed. Reason: Invalid UAnimNotify class"));
+					}
+				}
+			}
+		}
+	}
+	
+	GeneratedNotifyTracks.Reset();
+	ProcessedNotifyTracks.Reset();
+}
+
+void UFootstepAnimEventsModifier::ValidateNotifyTracks(UAnimSequence* InAnimation)
+{
+	check(InAnimation != nullptr)
+
+	GatherNotifyTracksInfo(InAnimation);
+	
+	for (const FFootDefinition & FootDef : FootDefinitions)
+	{
+		if (FootDef.bShouldGenerateSyncMarkers)
+		{
+			PrepareNotifyTrack(InAnimation, FootDef.SyncMarkerTrackName);
+		}
+		
+		if (FootDef.bShouldGenerateNotifies)
+		{
+			PrepareNotifyTrack(InAnimation, FootDef.FootstepNotifyTrackName);
+		}
+	}
+}
+
+void UFootstepAnimEventsModifier::GatherNotifyTracksInfo(const UAnimSequence* InAnimation)
+{
+	for (const FFootDefinition& FootDef : FootDefinitions)
+	{
+		// Determine tracks that will be generated
+		{
+			const bool bDoesRequestedSyncTrackAlreadyExist = FootDef.bShouldGenerateSyncMarkers && UAnimationBlueprintLibrary::IsValidAnimNotifyTrackName(InAnimation, FootDef.SyncMarkerTrackName);
+			const bool bDoesRequestedNotifyTrackAlreadyExist = FootDef.bShouldGenerateNotifies && UAnimationBlueprintLibrary::IsValidAnimNotifyTrackName(InAnimation, FootDef.FootstepNotifyTrackName);
+	
+			if (!bDoesRequestedSyncTrackAlreadyExist)
+			{
+				GeneratedNotifyTracks.Add(FootDef.SyncMarkerTrackName);
+			}
+
+			if (!bDoesRequestedNotifyTrackAlreadyExist)
+			{
+				GeneratedNotifyTracks.Add(FootDef.FootstepNotifyTrackName);
+			}
+		}
+
+		// Determine tracks that will be modified
+		ProcessedNotifyTracks.Add(FootDef.SyncMarkerTrackName);
+		ProcessedNotifyTracks.Add(FootDef.FootstepNotifyTrackName);
+	}
+}
+
+void UFootstepAnimEventsModifier::PrepareNotifyTrack(UAnimSequence* InAnimation, FName InRequestedNotifyTrackName)
+{
+	const bool bDoesTrackNameAlreadyExist = UAnimationBlueprintLibrary::IsValidAnimNotifyTrackName(InAnimation, InRequestedNotifyTrackName);
+	
+	if (!bDoesTrackNameAlreadyExist)
+	{
+		UAnimationBlueprintLibrary::AddAnimationNotifyTrack(InAnimation, InRequestedNotifyTrackName, FLinearColor::MakeRandomColor());
+	}
+	else if (bShouldRemovePreExistingNotifiesOrSyncMarkers)
+	{
+		UAnimationBlueprintLibrary::RemoveAnimationNotifyEventsByTrack(InAnimation, InRequestedNotifyTrackName);
+		UAnimationBlueprintLibrary::RemoveAnimationSyncMarkersByTrack(InAnimation, InRequestedNotifyTrackName);
+	}
+}
+
+bool UFootstepAnimEventsModifier::CanWePlaceEventAtSample(const FFootSampleState& InTestFootSampleState, EDetectionTechnique DetectionTechnique)
+{
+	switch(DetectionTechnique)
+	{
+		case EDetectionTechnique::PassThroughReferenceBone: return InTestFootSampleState.RefBoneTranslationDotRefBoneToFootBoneVec > 0.0f && InTestFootSampleState.PrevRefBoneTranslationDotRefBoneToFootBoneVec < 0.0f;
+		case EDetectionTechnique::FootBoneReachesGround: return !InTestFootSampleState.bWasFootBoneInGround && InTestFootSampleState.bIsFootBoneInGround;
+		default: return false;
+	}
+}
+#undef LOCTEXT_NAMESPACE
+
+#endif // WITH_EDITOR
diff --git a/Engine/Plugins/Animation/AnimationModifierLibrary/Source/AnimationModifierLibrary/Public/FootstepAnimEventsModifier.h b/Engine/Plugins/Animation/AnimationModifierLibrary/Source/AnimationModifierLibrary/Public/FootstepAnimEventsModifier.h
new file mode 100644
--- /dev/null
+++ b/Engine/Plugins/Animation/AnimationModifierLibrary/Source/AnimationModifierLibrary/Public/FootstepAnimEventsModifier.h
@@ -0,0 +1,128 @@
+﻿// Copyright Epic Games, Inc. All Rights Reserved.
+
+#pragma once
+
+#include "AnimationModifier.h"
+#include "FootstepAnimEventsModifier.generated.h"
+
+class UAnimNotify;
+
+/**  Detection method used for placing a notify / marker in the track */
+UENUM()
+enum class EDetectionTechnique : uint8
+{
+	/**
+	 * Create anim event when foot bone passes through a given reference bone.
+	 *
+	 * Note that the translation vector of the reference bone will be used as the normal vector for the detection plane
+	 * used to determine if the foot bone has crossed the reference bone. */
+	PassThroughReferenceBone,
+	
+	/**
+	 * Create anim event when foot bone reaches the ground level within a given threshold.
+	 * 
+	 * Note that this will only be true if the foot bone position goes from NOT being within the threshold to being
+	 * WITHIN the ground threshold.
+	 */
+	FootBoneReachesGround,
+};
+
+USTRUCT(BlueprintType)
+struct FFootDefinition
+{
+	GENERATED_BODY()
+	
+	UPROPERTY(EditAnywhere, Category = "Default")
+	FName FootBoneName = NAME_None;
+
+	UPROPERTY(EditAnywhere, BlueprintReadOnly, Category = "Default")
+	FName ReferenceBoneName = FName(TEXT("root"));
+	
+	UPROPERTY(EditAnywhere, Category = "Sync Markers")
+	bool bShouldGenerateSyncMarkers = false; 
+	
+	UPROPERTY(EditAnywhere, Category = "Sync Markers", meta=(EditCondition="bShouldGenerateSyncMarkers"))
+	FName SyncMarkerTrackName = TEXT("FootSyncMakers");
+
+	UPROPERTY(EditAnywhere, Category = "Sync Markers", meta=(EditCondition="bShouldGenerateSyncMarkers"))
+	FName SyncMarkerName = NAME_None;
+
+	UPROPERTY(EditAnywhere, Category = "Sync Markers", meta=(EditCondition="bShouldGenerateSyncMarkers"))
+	EDetectionTechnique SyncMarkerDetectionTechnique = EDetectionTechnique::PassThroughReferenceBone;
+	
+	UPROPERTY(EditAnywhere, Category = "Notifies")
+	bool bShouldGenerateNotifies = false;
+	
+	UPROPERTY(EditAnywhere, Category = "Notifies", meta=(EditCondition="bShouldGenerateNotifies"))
+	FName FootstepNotifyTrackName = TEXT("FootAnimEvents");
+
+	UPROPERTY(EditAnywhere, Category = "Notifies", meta=(EditCondition="bShouldGenerateNotifies"))
+	TSubclassOf<UAnimNotify> FootstepNotify = nullptr;
+
+	UPROPERTY(EditAnywhere, Category = "Notifies", meta=(EditCondition="bShouldGenerateNotifies"))
+	EDetectionTechnique FootstepNotifyDetectionTechnique = EDetectionTechnique::FootBoneReachesGround;
+};
+
+/** Generates animation notifies and/or sync markers for any specified bone(s) */
+UCLASS(meta=(IsBlueprintBase=true))
+class UFootstepAnimEventsModifier : public UAnimationModifier
+{
+	GENERATED_BODY()
+
+public:
+
+	/** Delta of each sampling step */
+	UPROPERTY(EditAnywhere, BlueprintReadOnly, Category = "Default")
+	float SampleStep;
+
+	/** Threshold for determining if a foot bone position can be considered to be on the ground level */
+	UPROPERTY(EditAnywhere, BlueprintReadOnly, Category = "Default")
+	float GroundThreshold;
+	
+	/**
+	 * If true, applying the anim modifier becomes a destructive action, meaning that any existing matched tracks will have their data overwritten by the modifier.
+	 * Otherwise, no previous notifies or sync markers will removed when applying the anim modifier.
+	 */
+	UPROPERTY(EditAnywhere, BlueprintReadOnly, Category = "Default")
+	bool bShouldRemovePreExistingNotifiesOrSyncMarkers;
+
+	/** Foot bone(s) to be processed */
+	UPROPERTY(EditAnywhere, BlueprintReadOnly, Category = "Default")
+	TArray<FFootDefinition> FootDefinitions;
+	
+	UFootstepAnimEventsModifier();
+
+	virtual void OnApply_Implementation(UAnimSequence* InAnimation) override;
+	virtual void OnRevert_Implementation(UAnimSequence* InAnimation) override;
+
+private:
+
+	/** Store state of a foot during a sample step */
+	struct FFootSampleState
+	{
+		double GroundLevel = UE_DOUBLE_BIG_NUMBER;
+		double RefBoneTranslationDotRefBoneToFootBoneVec = 0.0;
+		double PrevRefBoneTranslationDotRefBoneToFootBoneVec = 0.0;
+		bool bIsFootBoneInGround = false;
+		bool bWasFootBoneInGround = false; 
+	};
+	
+	/** Keep track of to be generated tracks during modifier application */
+	TSet<FName> GeneratedNotifyTracks;
+	TSet<FName> ProcessedNotifyTracks;
+	
+	/** Prepare all requested notify tracks before starting to generate notifies or sync markers */
+	void ValidateNotifyTracks(UAnimSequence* InAnimation);
+
+	/** Determine which notify tracks will be generated when this modifier is applied */
+	void GatherNotifyTracksInfo(const UAnimSequence* InAnimation);
+
+	/** Create new notify track if requested track does not exist, otherwise, if the requested track exist, clean-up data if needed. */
+	void PrepareNotifyTrack(UAnimSequence* InAnimation, FName InRequestedNotifyTrackName);
+
+	/** Test if we can place anim event at given sample */
+	static bool CanWePlaceEventAtSample(const FFootSampleState & InTestLegSampleState, EDetectionTechnique DetectionTechnique);
+
+	/** Used to keep track of state of locked root motion before applying modifier */
+	bool bIsRootMotionLocked;
+};
\ No newline at end of file
