Fixes / updates for FootstepAnimEventsModifier:

- Improved reverting of anim modifier.
- Fixed a typo on the track name for sync markers.
- Added a better foot bone definitions item visualization on editor.

diff --git a/Engine/Plugins/Animation/AnimationModifierLibrary/Source/AnimationModifierLibrary/Private/FootstepAnimEventsModifier.cpp b/Engine/Plugins/Animation/AnimationModifierLibrary/Source/AnimationModifierLibrary/Private/FootstepAnimEventsModifier.cpp
--- a/Engine/Plugins/Animation/AnimationModifierLibrary/Source/AnimationModifierLibrary/Private/FootstepAnimEventsModifier.cpp
+++ b/Engine/Plugins/Animation/AnimationModifierLibrary/Source/AnimationModifierLibrary/Private/FootstepAnimEventsModifier.cpp
@@ -214,25 +214,29 @@ void UFootstepAnimEventsModifier::GatherNotifyTracksInfo(const UAnimSequence* In
 {
 	for (const FFootDefinition& FootDef : FootDefinitions)
 	{
-		// Determine tracks that will be generated
+		// Determine tracks that will be generated and/or processed
+		const bool bDoesRequestedSyncTrackAlreadyExist = UAnimationBlueprintLibrary::IsValidAnimNotifyTrackName(InAnimation, FootDef.SyncMarkerTrackName);
+		const bool bDoesRequestedNotifyTrackAlreadyExist = UAnimationBlueprintLibrary::IsValidAnimNotifyTrackName(InAnimation, FootDef.FootstepNotifyTrackName);
+
+		if (FootDef.bShouldGenerateSyncMarkers)
 		{
-			const bool bDoesRequestedSyncTrackAlreadyExist = FootDef.bShouldGenerateSyncMarkers && UAnimationBlueprintLibrary::IsValidAnimNotifyTrackName(InAnimation, FootDef.SyncMarkerTrackName);
-			const bool bDoesRequestedNotifyTrackAlreadyExist = FootDef.bShouldGenerateNotifies && UAnimationBlueprintLibrary::IsValidAnimNotifyTrackName(InAnimation, FootDef.FootstepNotifyTrackName);
-	
 			if (!bDoesRequestedSyncTrackAlreadyExist)
 			{
 				GeneratedNotifyTracks.Add(FootDef.SyncMarkerTrackName);
 			}
 
+			ProcessedNotifyTracks.Add(FootDef.SyncMarkerTrackName);
+		}
+
+		if (FootDef.bShouldGenerateNotifies)
+		{
 			if (!bDoesRequestedNotifyTrackAlreadyExist)
 			{
 				GeneratedNotifyTracks.Add(FootDef.FootstepNotifyTrackName);
 			}
-		}
 
-		// Determine tracks that will be modified
-		ProcessedNotifyTracks.Add(FootDef.SyncMarkerTrackName);
-		ProcessedNotifyTracks.Add(FootDef.FootstepNotifyTrackName);
+			ProcessedNotifyTracks.Add(FootDef.FootstepNotifyTrackName);
+		}
 	}
 }
 
diff --git a/Engine/Plugins/Animation/AnimationModifierLibrary/Source/AnimationModifierLibrary/Public/FootstepAnimEventsModifier.h b/Engine/Plugins/Animation/AnimationModifierLibrary/Source/AnimationModifierLibrary/Public/FootstepAnimEventsModifier.h
--- a/Engine/Plugins/Animation/AnimationModifierLibrary/Source/AnimationModifierLibrary/Public/FootstepAnimEventsModifier.h
+++ b/Engine/Plugins/Animation/AnimationModifierLibrary/Source/AnimationModifierLibrary/Public/FootstepAnimEventsModifier.h
@@ -42,7 +42,7 @@ struct FFootDefinition
 	bool bShouldGenerateSyncMarkers = false; 
 	
 	UPROPERTY(EditAnywhere, Category = "Sync Markers", meta=(EditCondition="bShouldGenerateSyncMarkers"))
-	FName SyncMarkerTrackName = TEXT("FootSyncMakers");
+	FName SyncMarkerTrackName = TEXT("FootSyncMarkers");
 
 	UPROPERTY(EditAnywhere, Category = "Sync Markers", meta=(EditCondition="bShouldGenerateSyncMarkers"))
 	FName SyncMarkerName = NAME_None;
@@ -72,24 +72,25 @@ class UFootstepAnimEventsModifier : public UAnimationModifier
 public:
 
 	/** Delta of each sampling step */
-	UPROPERTY(EditAnywhere, BlueprintReadOnly, Category = "Default")
+	UPROPERTY(EditAnywhere, BlueprintReadOnly, Category = "Settings")
 	float SampleStep;
 
 	/** Threshold for determining if a foot bone position can be considered to be on the ground level */
-	UPROPERTY(EditAnywhere, BlueprintReadOnly, Category = "Default")
+	UPROPERTY(EditAnywhere, BlueprintReadOnly, Category = "Settings")
 	float GroundThreshold;
+
+	/** Foot bone(s) to be processed */
+	UPROPERTY(EditAnywhere, BlueprintReadOnly, Category = "Settings", meta = (TitleProperty = "FootBoneName"))
+	TArray<FFootDefinition> FootDefinitions;
 	
 	/**
 	 * If true, applying the anim modifier becomes a destructive action, meaning that any existing matched tracks will have their data overwritten by the modifier.
 	 * Otherwise, no previous notifies or sync markers will removed when applying the anim modifier.
 	 */
-	UPROPERTY(EditAnywhere, BlueprintReadOnly, Category = "Default")
+	UPROPERTY(EditAnywhere, BlueprintReadOnly, Category = "Settings", DisplayName="Remove Pre Existing Notifies or Sync Markers")
 	bool bShouldRemovePreExistingNotifiesOrSyncMarkers;
 
-	/** Foot bone(s) to be processed */
-	UPROPERTY(EditAnywhere, BlueprintReadOnly, Category = "Default")
-	TArray<FFootDefinition> FootDefinitions;
-	
+
 	UFootstepAnimEventsModifier();
 
 	virtual void OnApply_Implementation(UAnimSequence* InAnimation) override;
@@ -108,7 +109,11 @@ private:
 	};
 	
 	/** Keep track of to be generated tracks during modifier application */
+	UPROPERTY(VisibleAnywhere, AdvancedDisplay, meta = (Category = "Settings"))
 	TSet<FName> GeneratedNotifyTracks;
+
+	/** Keep track of tracks modified during modifier application */
+	UPROPERTY(VisibleAnywhere, AdvancedDisplay, meta = (Category = "Settings"))
 	TSet<FName> ProcessedNotifyTracks;
 	
 	/** Prepare all requested notify tracks before starting to generate notifies or sync markers */
