Added "FootBoneSpeed" as an additional detection method for placing foot step anim events in `FootstepAnimEventsModifier`.

- Removed `SampleStep` variable in favor of `SampleRate` so its easier to understand for the end user.

[RN] rnx

diff --git a/Engine/Plugins/Animation/AnimationModifierLibrary/Source/AnimationModifierLibrary/Private/FootstepAnimEventsModifier.cpp b/Engine/Plugins/Animation/AnimationModifierLibrary/Source/AnimationModifierLibrary/Private/FootstepAnimEventsModifier.cpp
--- a/Engine/Plugins/Animation/AnimationModifierLibrary/Source/AnimationModifierLibrary/Private/FootstepAnimEventsModifier.cpp
+++ b/Engine/Plugins/Animation/AnimationModifierLibrary/Source/AnimationModifierLibrary/Private/FootstepAnimEventsModifier.cpp
@@ -13,8 +13,9 @@
 
 UFootstepAnimEventsModifier::UFootstepAnimEventsModifier() : Super()
 {
-	SampleStep = 0.0125f;
-	GroundThreshold = 0.95f;
+	SampleRate = 60;
+	GroundThreshold = 4.0f;
+	SpeedThreshold = 0.1f;
 	bShouldRemovePreExistingNotifiesOrSyncMarkers = false;
 }
 
@@ -29,122 +30,132 @@ void UFootstepAnimEventsModifier::OnApply_Implementation(UAnimSequence* InAnimat
 	}
 	
 	// Disable root motion lock during application
-	bIsRootMotionLocked = InAnimation->bForceRootLock;
-	InAnimation->bForceRootLock = false;
-
+	TGuardValue<bool> ForceRootLockGuard(InAnimation->bForceRootLock, false);
+	
 	// Clean up or generate tracks if needed
 	ValidateNotifyTracks(InAnimation);
 	
-	// Process animation
+	// Process animation asset
 	{
+		const FAnimPoseEvaluationOptions AnimPoseEvalOptions { EAnimDataEvalType::Raw, true, false, false, nullptr, true, false };
+		const float SequenceLength = InAnimation->GetPlayLength();
+		const float SampleStep = 1.0f / static_cast<float>(SampleRate);
+		const int SampleNum = FMath::TruncToInt(SequenceLength / SampleStep);
+		
 		TArray<FFootSampleState> FootSampleStates;
 		FootSampleStates.Init(FFootSampleState(), FootDefinitions.Num());
 		
-		const float SequenceLength = InAnimation->GetPlayLength();
-		const int SampleNum = FMath::TruncToInt(SequenceLength / SampleStep);
-
-		// Get ground levels for each foot definition
-		for (int i = 0; i < SampleNum; ++i)
+		// Get ground levels and max speed values.
+		for (int SampleIndex = 0; SampleIndex < SampleNum; ++SampleIndex)
 		{
-			const float SampleTime = static_cast<float>(i) * SampleStep;
-			FAnimPose AnimPose;
-		
-			// Query animation pose at the current sample time
-			{
-				FAnimPoseEvaluationOptions AnimPoseEvalOptions;
-				AnimPoseEvalOptions.EvaluationType = EAnimDataEvalType::Raw;
-			
-				UAnimPoseExtensions::GetAnimPoseAtTime(InAnimation, SampleTime, AnimPoseEvalOptions, AnimPose);
-			}
-
-			// Find out ground level for each foot
-			for (int CurrFootIdx = 0; CurrFootIdx < FootDefinitions.Num(); ++CurrFootIdx)
-			{
-				const FFootDefinition & FootDef = FootDefinitions[CurrFootIdx];
-				FFootSampleState & FootState = FootSampleStates[CurrFootIdx];
-
-				// Compute ground level
-				FootState.GroundLevel = FMath::Min(FootState.GroundLevel, UAnimPoseExtensions::GetBonePose(AnimPose, FootDef.FootBoneName, EAnimPoseSpaces::World).GetLocation().Z);
-			}
-		}
+			const float SampleTime = FMath::Clamp(static_cast<float>(SampleIndex) * SampleStep, 0.0f, SequenceLength);
+			const float FutureSampleTime = FMath::Clamp((static_cast<float>(SampleIndex) + 1.0f) * SampleStep, 0.0f, SequenceLength);
 
-		// Avoid creating a notify at the very start of a sequence if the foot bone is already within ground threshold 
-		{
 			FAnimPose AnimPose;
+			FAnimPose FutureAnimPose;
 			
-			FAnimPoseEvaluationOptions AnimPoseEvalOptions;
-			AnimPoseEvalOptions.EvaluationType = EAnimDataEvalType::Raw;
-			
-			UAnimPoseExtensions::GetAnimPoseAtTime(InAnimation, 0.0f, AnimPoseEvalOptions, AnimPose);
+			UAnimPoseExtensions::GetAnimPoseAtTime(InAnimation, SampleTime, AnimPoseEvalOptions, AnimPose);
+			UAnimPoseExtensions::GetAnimPoseAtTime(InAnimation, FutureSampleTime, AnimPoseEvalOptions, FutureAnimPose);
 			
-			for (int CurrFootIdx = 0; CurrFootIdx < FootDefinitions.Num(); ++CurrFootIdx)
+			for (int FootIndex = 0; FootIndex < FootDefinitions.Num(); ++FootIndex)
 			{
-				const FFootDefinition & FootDef = FootDefinitions[CurrFootIdx];
-				FFootSampleState & FootState = FootSampleStates[CurrFootIdx];
+				const FFootDefinition & FootDef = FootDefinitions[FootIndex];
+				FFootSampleState & FootState = FootSampleStates[FootIndex];
 				
-				FootState.bWasFootBoneInGround = FMath::Abs(FootState.GroundLevel - UAnimPoseExtensions::GetBonePose(AnimPose, FootDef.FootBoneName, EAnimPoseSpaces::World).GetLocation().Z) < GroundThreshold;
+				FootState.GroundLevel = FMath::Min(FootState.GroundLevel, UAnimPoseExtensions::GetBonePose(AnimPose, FootDef.FootBoneName, EAnimPoseSpaces::World).GetLocation().Z);
+				FootState.MaxFootSpeed = FMath::Max(FootState.MaxFootSpeed, FMath::Abs(ComputeBoneSpeed(AnimPose, FutureAnimPose, SampleStep, FootDef.FootBoneName)));
 			}
 		}
-		
+
 		// Generate animation events
-		for (int i = 0; i < SampleNum; ++i)
+		for (int SampleIndex = 0; SampleIndex < (SampleNum - 1); ++SampleIndex)
 		{
-			const float SampleTime = static_cast<float>(i) * SampleStep;
+			const float SampleTime = FMath::Clamp(static_cast<float>(SampleIndex) * SampleStep, 0.0f, SequenceLength);
+			const float FutureSampleTime = FMath::Clamp((static_cast<float>(SampleIndex) + 1.0f) * SampleStep, 0.0f, SequenceLength);
 
 			FAnimPose AnimPose;
 			FAnimPose FutureAnimPose;
 			
 			// Query animation pose at the current sample time
-			{
-				FAnimPoseEvaluationOptions AnimPoseEvalOptions;
-				AnimPoseEvalOptions.EvaluationType = EAnimDataEvalType::Raw;
-				
-				UAnimPoseExtensions::GetAnimPoseAtTime(InAnimation, SampleTime, AnimPoseEvalOptions, AnimPose);
-				UAnimPoseExtensions::GetAnimPoseAtTime(InAnimation, FMath::Clamp(SampleTime + SampleStep, 0.0f, SequenceLength), AnimPoseEvalOptions, FutureAnimPose);
-			}
-
+			UAnimPoseExtensions::GetAnimPoseAtTime(InAnimation, SampleTime, AnimPoseEvalOptions, AnimPose);
+			UAnimPoseExtensions::GetAnimPoseAtTime(InAnimation, FutureSampleTime, AnimPoseEvalOptions, FutureAnimPose);
+			
 			// Process all foot definitions
 			for (int CurrFootIdx = 0; CurrFootIdx < FootDefinitions.Num(); ++CurrFootIdx)
 			{
 				const FFootDefinition & FootDef = FootDefinitions[CurrFootIdx];
 				FFootSampleState & FootState = FootSampleStates[CurrFootIdx];
 				
-				// Update current sample state
+				// Update current sample info
 				{
-					// Get reference bone translation
-					const FTransform ReferenceBoneTransform = UAnimPoseExtensions::GetBonePose(AnimPose, FootDef.ReferenceBoneName, EAnimPoseSpaces::World);
-					const FTransform FutureReferenceBoneTransform = UAnimPoseExtensions::GetBonePose(FutureAnimPose, FootDef.ReferenceBoneName, EAnimPoseSpaces::World);
-					const FVector ReferenceBoneTranslation = FutureReferenceBoneTransform.GetLocation() - ReferenceBoneTransform.GetLocation();
-
-					// Get foot bone direction with respect to reference bone
 					const FTransform FootBoneTransform = UAnimPoseExtensions::GetBonePose(AnimPose, FootDef.FootBoneName, EAnimPoseSpaces::World);
-					const FVector ReferenceBoneToFootBoneVector = FootBoneTransform.GetLocation() - ReferenceBoneTransform.GetLocation();
 
-					FootState.RefBoneTranslationDotRefBoneToFootBoneVec = FVector::DotProduct(ReferenceBoneTranslation, ReferenceBoneToFootBoneVector);
-					FootState.bIsFootBoneInGround = FMath::Abs(FootState.GroundLevel - FootBoneTransform.GetLocation().Z) < GroundThreshold;
+					// Method: PassThroughReferenceBone
+					{
+						// Get reference bone translation
+						const FTransform ReferenceBoneTransform = UAnimPoseExtensions::GetBonePose(AnimPose, FootDef.ReferenceBoneName, EAnimPoseSpaces::World);
+						const FTransform FutureReferenceBoneTransform = UAnimPoseExtensions::GetBonePose(FutureAnimPose, FootDef.ReferenceBoneName, EAnimPoseSpaces::World);
+						const FVector ReferenceBoneTranslation = FutureReferenceBoneTransform.GetLocation() - ReferenceBoneTransform.GetLocation();
+
+						// Get foot bone direction with respect to reference bone
+						const FVector ReferenceBoneToFootBoneVector = FootBoneTransform.GetLocation() - ReferenceBoneTransform.GetLocation();
+					
+						FootState.RefBoneTranslationDotRefBoneToFootBoneVec = FVector::DotProduct(ReferenceBoneTranslation, ReferenceBoneToFootBoneVector);
+					}
+
+					// Method: FootBoneReachesGround
+					FootState.bIsFootBoneInGround = FMath::Abs(FootState.GroundLevel - FootBoneTransform.GetLocation().Z) <= GroundThreshold;
+
+					// Method: FootBoneSpeed
+					FootState.FootBoneSpeed = FMath::Abs(ComputeBoneSpeed(AnimPose, FutureAnimPose, SampleStep, FootDef.FootBoneName)) / FootState.MaxFootSpeed;
+
+					// Keep track of sample with lowest speed below threshold
+					if (FootState.FootBoneSpeed < SpeedThreshold)
+					{
+						if (FootState.FootBoneSpeed < FootState.MinFootSpeedBelowThreshold)
+						{
+							FootState.TimeAtMinFootSpeedBelowThreshold = SampleTime;
+							FootState.MinFootSpeedBelowThreshold = FootState.FootBoneSpeed;
+						}
+					}
 				}
 				
-				// Generate sync markers
-				if (FootDef.bShouldGenerateSyncMarkers && CanWePlaceEventAtSample(FootState, FootDef.SyncMarkerDetectionTechnique))
+				if (SampleIndex > 0)
 				{
-					UAnimationBlueprintLibrary::AddAnimationSyncMarker(InAnimation, FootDef.SyncMarkerName, SampleTime, FootDef.SyncMarkerTrackName);
-				}
+					// Generate sync markers
+					if (FootDef.bShouldGenerateSyncMarkers && CanWePlaceEventAtSample(FootState, FootDef.SyncMarkerDetectionTechnique))
+					{
+						float FinalSyncMarkerTime = FootDef.SyncMarkerDetectionTechnique == EDetectionTechnique::FootBoneSpeed ? FootState.TimeAtMinFootSpeedBelowThreshold : SampleTime;
+						UAnimationBlueprintLibrary::AddAnimationSyncMarker(InAnimation, FootDef.SyncMarkerName, FinalSyncMarkerTime, FootDef.SyncMarkerTrackName);
+					}
 
-				// Generate foot step fx notifies
-				if (FootDef.bShouldGenerateNotifies && CanWePlaceEventAtSample(FootState, FootDef.FootstepNotifyDetectionTechnique))
+					// Generate foot step fx notifies
+					if (FootDef.bShouldGenerateNotifies && CanWePlaceEventAtSample(FootState, FootDef.FootstepNotifyDetectionTechnique))
+					{
+						float FinalAnimNotifyTime = FootDef.FootstepNotifyDetectionTechnique == EDetectionTechnique::FootBoneSpeed ? FootState.TimeAtMinFootSpeedBelowThreshold : SampleTime;
+						UAnimationBlueprintLibrary::AddAnimationNotifyEvent(InAnimation, FootDef.FootstepNotifyTrackName, FinalAnimNotifyTime, FootDef.FootstepNotify);
+					}
+				}
+				
+				// Update foot state
 				{
-					UAnimationBlueprintLibrary::AddAnimationNotifyEvent(InAnimation, FootDef.FootstepNotifyTrackName, SampleTime, FootDef.FootstepNotify);
-				}					
+					const bool bDidWePlaceAnyEventsUsingFootBoneSpeed = CanWePlaceEventAtSample(FootState, EDetectionTechnique::FootBoneSpeed);
 
-				// Update foot state
-				FootState.bWasFootBoneInGround = FootState.bIsFootBoneInGround;
-				FootState.PrevRefBoneTranslationDotRefBoneToFootBoneVec = FootState.RefBoneTranslationDotRefBoneToFootBoneVec;
+					// Reset minimum values if we are above speed threshold
+					if (bDidWePlaceAnyEventsUsingFootBoneSpeed)
+					{
+						FootState.TimeAtMinFootSpeedBelowThreshold = MAX_FLT;
+						FootState.MinFootSpeedBelowThreshold = MAX_FLT;
+					}
+					
+					// Keep track of foot info
+					FootState.PrevFootBoneSpeed = FootState.FootBoneSpeed;
+					FootState.bWasFootBoneInGround = FootState.bIsFootBoneInGround;
+					FootState.PrevRefBoneTranslationDotRefBoneToFootBoneVec = FootState.RefBoneTranslationDotRefBoneToFootBoneVec;
+				}
 			}
 		}
 	}
-
-	// Revert modifications to animation state
-	InAnimation->bForceRootLock = bIsRootMotionLocked;
 }
 
 void UFootstepAnimEventsModifier::OnRevert_Implementation(UAnimSequence* InAnimation)
@@ -255,15 +266,27 @@ void UFootstepAnimEventsModifier::PrepareNotifyTrack(UAnimSequence* InAnimation,
 	}
 }
 
-bool UFootstepAnimEventsModifier::CanWePlaceEventAtSample(const FFootSampleState& InTestFootSampleState, EDetectionTechnique DetectionTechnique)
+bool UFootstepAnimEventsModifier::CanWePlaceEventAtSample(const FFootSampleState& InFootSampleState, EDetectionTechnique DetectionTechnique) const
 {
 	switch(DetectionTechnique)
 	{
-		case EDetectionTechnique::PassThroughReferenceBone: return InTestFootSampleState.RefBoneTranslationDotRefBoneToFootBoneVec > 0.0f && InTestFootSampleState.PrevRefBoneTranslationDotRefBoneToFootBoneVec < 0.0f;
-		case EDetectionTechnique::FootBoneReachesGround: return !InTestFootSampleState.bWasFootBoneInGround && InTestFootSampleState.bIsFootBoneInGround;
+		case EDetectionTechnique::PassThroughReferenceBone: return InFootSampleState.RefBoneTranslationDotRefBoneToFootBoneVec > 0.0f && InFootSampleState.PrevRefBoneTranslationDotRefBoneToFootBoneVec < 0.0f;
+		case EDetectionTechnique::FootBoneReachesGround: return !InFootSampleState.bWasFootBoneInGround && InFootSampleState.bIsFootBoneInGround;
+		case EDetectionTechnique::FootBoneSpeed: return InFootSampleState.PrevFootBoneSpeed < SpeedThreshold && InFootSampleState.FootBoneSpeed >= SpeedThreshold;
 		default: return false;
 	}
 }
+
+float UFootstepAnimEventsModifier::ComputeBoneSpeed(const FAnimPose& InPose, const FAnimPose& InFuturePose, float InDelta, FName InFootBoneName) const
+{
+	check(!FMath::IsNearlyZero(InDelta))
+			
+	const FTransform FootBoneTransform = UAnimPoseExtensions::GetBonePose(InPose, InFootBoneName, EAnimPoseSpaces::World);
+	const FTransform FutureFootBoneTransform = UAnimPoseExtensions::GetBonePose(InFuturePose, InFootBoneName, EAnimPoseSpaces::World);
+	const double FootBoneDistance = (FutureFootBoneTransform.GetLocation() - FootBoneTransform.GetLocation()).Length();
+			
+	return static_cast<float>(FootBoneDistance / InDelta);
+}
 #undef LOCTEXT_NAMESPACE
 
 #endif // WITH_EDITOR
diff --git a/Engine/Plugins/Animation/AnimationModifierLibrary/Source/AnimationModifierLibrary/Public/FootstepAnimEventsModifier.h b/Engine/Plugins/Animation/AnimationModifierLibrary/Source/AnimationModifierLibrary/Public/FootstepAnimEventsModifier.h
--- a/Engine/Plugins/Animation/AnimationModifierLibrary/Source/AnimationModifierLibrary/Public/FootstepAnimEventsModifier.h
+++ b/Engine/Plugins/Animation/AnimationModifierLibrary/Source/AnimationModifierLibrary/Public/FootstepAnimEventsModifier.h
@@ -3,6 +3,7 @@
 #pragma once
 
 #include "AnimationModifier.h"
+#include "AnimPose.h"
 #include "FootstepAnimEventsModifier.generated.h"
 
 class UAnimNotify;
@@ -25,6 +26,14 @@ enum class EDetectionTechnique : uint8
 	 * WITHIN the ground threshold.
 	 */
 	FootBoneReachesGround,
+
+	/**
+	 * Create anim event when foot bone translation speed is below a given threshold and nearly zero.
+	 *
+	 * Note that the foot bone translation speed is normalize therefore when a footstep occurs
+	 * the speed will be very close to zero.
+	 */
+	FootBoneSpeed
 };
 
 USTRUCT(BlueprintType)
@@ -71,14 +80,25 @@ class UFootstepAnimEventsModifier : public UAnimationModifier
 
 public:
 
-	/** Delta of each sampling step */
-	UPROPERTY(EditAnywhere, BlueprintReadOnly, Category = "Settings")
-	float SampleStep;
+	UFootstepAnimEventsModifier();
+	
+	/** Rate used to sample the animation */
+	UPROPERTY(EditAnywhere, BlueprintReadOnly, Category = "Settings", meta = (Units="Hz", UIMin=1))
+	int SampleRate;
 
 	/** Threshold for determining if a foot bone position can be considered to be on the ground level */
 	UPROPERTY(EditAnywhere, BlueprintReadOnly, Category = "Settings")
 	float GroundThreshold;
 
+	/**
+	 * Threshold to start finding the smallest foot bone translation speed.
+	 *
+	 * Note that the foot bone translation speed is normalize therefore when a footstep occurs
+	 * the speed will be very close to zero, thus for most cases this value won't need to be changed.
+	 */
+	UPROPERTY(EditAnywhere, BlueprintReadOnly, Category = "Settings", meta = (ClampMin=0.0f, ClampMax=1.0f))
+	float SpeedThreshold;
+	
 	/** Foot bone(s) to be processed */
 	UPROPERTY(EditAnywhere, BlueprintReadOnly, Category = "Settings", meta = (TitleProperty = "FootBoneName"))
 	TArray<FFootDefinition> FootDefinitions;
@@ -90,11 +110,10 @@ public:
 	UPROPERTY(EditAnywhere, BlueprintReadOnly, Category = "Settings", DisplayName="Remove Pre Existing Notifies or Sync Markers")
 	bool bShouldRemovePreExistingNotifiesOrSyncMarkers;
 
-
-	UFootstepAnimEventsModifier();
-
+	/** Begin UAnimationModifier interface */
 	virtual void OnApply_Implementation(UAnimSequence* InAnimation) override;
 	virtual void OnRevert_Implementation(UAnimSequence* InAnimation) override;
+	/** End UAnimationModifier interface */
 
 private:
 
@@ -106,6 +125,11 @@ private:
 		double PrevRefBoneTranslationDotRefBoneToFootBoneVec = 0.0;
 		bool bIsFootBoneInGround = false;
 		bool bWasFootBoneInGround = false; 
+		float FootBoneSpeed = 0.0f;
+		float PrevFootBoneSpeed = 0.0f;
+		float MaxFootSpeed = -MAX_FLT;
+		float MinFootSpeedBelowThreshold = MAX_FLT;
+		float TimeAtMinFootSpeedBelowThreshold = MAX_FLT;
 	};
 	
 	/** Keep track of to be generated tracks during modifier application */
@@ -126,8 +150,8 @@ private:
 	void PrepareNotifyTrack(UAnimSequence* InAnimation, FName InRequestedNotifyTrackName);
 
 	/** Test if we can place anim event at given sample */
-	static bool CanWePlaceEventAtSample(const FFootSampleState & InTestLegSampleState, EDetectionTechnique DetectionTechnique);
+	bool CanWePlaceEventAtSample(const FFootSampleState & InTestLegSampleState, EDetectionTechnique DetectionTechnique) const;
 
-	/** Used to keep track of state of locked root motion before applying modifier */
-	bool bIsRootMotionLocked;
+	/** Get a bone translation speed */
+	float ComputeBoneSpeed(const FAnimPose& InPose, const FAnimPose &InFuturePose, float InDelta, FName InFootBoneName) const;
 };
\ No newline at end of file
