MotionExtractor: Added space enum to allow calculating motion relative to another bone.

diff --git a/Engine/Plugins/Animation/AnimationModifierLibrary/Source/AnimationModifierLibrary/Private/MotionExtractorModifier.cpp b/Engine/Plugins/Animation/AnimationModifierLibrary/Source/AnimationModifierLibrary/Private/MotionExtractorModifier.cpp
--- a/Engine/Plugins/Animation/AnimationModifierLibrary/Source/AnimationModifierLibrary/Private/MotionExtractorModifier.cpp
+++ b/Engine/Plugins/Animation/AnimationModifierLibrary/Source/AnimationModifierLibrary/Private/MotionExtractorModifier.cpp
@@ -20,7 +20,7 @@ UMotionExtractorModifier::UMotionExtractorModifier()
 	Axis = EMotionExtractor_Axis::Y;
 	bAbsoluteValue = false;
 	bRelativeToFirstFrame = false;
-	bComponentSpace = true;
+	bComponentSpace_DEPRECATED = true;
 	bNormalize = false;
 	bUseCustomCurveName = false;
 	CustomCurveName = NAME_None;
@@ -52,6 +52,22 @@ void UMotionExtractorModifier::OnApply_Implementation(UAnimSequence* Animation)
 		return;
 	}
 
+	int32 RelativeToBoneIndex = INDEX_NONE;
+	const bool bUseRelativeToBone = Space == EMotionExtractor_Space::RelativeToBone;
+	if (bUseRelativeToBone)
+	{
+		RelativeToBoneIndex = Skeleton->GetReferenceSkeleton().FindBoneIndex(RelativeToBoneName);
+		if (RelativeToBoneIndex == INDEX_NONE)
+		{
+			UE_LOG(LogAnimation, Error, TEXT("MotionExtractorModifier failed. Reason: Invalid Bone Index for RelativeToBoneName. BoneName: %s Animation: %s Skeleton: %s"),
+				*BoneName.ToString(), *GetNameSafe(Animation), *GetNameSafe(Skeleton));
+			return;
+		}
+	}
+
+	// When Space == EMotionExtractor_Space::RelativeToBone, operations only make sense to do in component space.
+	const bool bShouldUseComponentSpace = Space != EMotionExtractor_Space::LocalSpace;
+
 	// Ideally we would disable these options when any of those motion types are selected but AnimModifier doesn't support Details Customization atm.
 	if((MotionType == EMotionExtractor_MotionType::Translation || MotionType == EMotionExtractor_MotionType::Rotation || MotionType == EMotionExtractor_MotionType::Scale) && Axis > EMotionExtractor_Axis::Z)
 	{
@@ -65,12 +81,25 @@ void UMotionExtractorModifier::OnApply_Implementation(UAnimSequence* Animation)
 
 	TArray<FBoneIndexType> RequiredBones;
 	RequiredBones.Add(IntCastChecked<FBoneIndexType>(BoneIndex));
-	Skeleton->GetReferenceSkeleton().EnsureParentsExistAndSort(RequiredBones);
 
+	if (bUseRelativeToBone)
+	{
+		RequiredBones.Add(IntCastChecked<FBoneIndexType>(RelativeToBoneIndex));
+	}
+
+	Skeleton->GetReferenceSkeleton().EnsureParentsExistAndSort(RequiredBones);
 	FBoneContainer BoneContainer(RequiredBones, UE::Anim::ECurveFilterMode::DisallowAll, *Skeleton);
 	const FCompactPoseBoneIndex CompactPoseBoneIndex = BoneContainer.MakeCompactPoseIndex(FMeshPoseBoneIndex(BoneIndex));
 
-	FTransform FirstFrameBoneTransform = UMotionExtractorUtilityLibrary::ExtractBoneTransform(Animation, BoneContainer, CompactPoseBoneIndex, 0.f, bComponentSpace);
+	FCompactPoseBoneIndex CompactPoseRelativeToBoneIndex = FCompactPoseBoneIndex(INDEX_NONE);
+	FTransform FirstFrameRelativeToBoneTransform = FTransform::Identity;
+	if (bUseRelativeToBone)
+	{
+		CompactPoseRelativeToBoneIndex = BoneContainer.MakeCompactPoseIndex(FMeshPoseBoneIndex(RelativeToBoneIndex));
+	}
+
+	// If relative to first frame, use RelativeToBone's first frame as the reference frame.
+	FTransform FirstFrameBoneTransform = UMotionExtractorUtilityLibrary::ExtractBoneTransform(Animation, BoneContainer, bUseRelativeToBone ? CompactPoseRelativeToBoneIndex : CompactPoseBoneIndex, 0.f, bShouldUseComponentSpace);
 
 	const float AnimLength = Animation->GetPlayLength();
 	const float SampleInterval = 1.f / static_cast<float>(SampleRate);
@@ -87,12 +116,17 @@ void UMotionExtractorModifier::OnApply_Implementation(UAnimSequence* Animation)
 			Time = FMath::Clamp(static_cast<float>(SampleIndex) * SampleInterval, 0.f, AnimLength);
 			SampleIndex++;
 
-			FTransform BoneTransform = UMotionExtractorUtilityLibrary::ExtractBoneTransform(Animation, BoneContainer, CompactPoseBoneIndex, Time, bComponentSpace);
+			FTransform BoneTransform = UMotionExtractorUtilityLibrary::ExtractBoneTransform(Animation, BoneContainer, CompactPoseBoneIndex, Time, bShouldUseComponentSpace);
 
 			if(bRelativeToFirstFrame)
 			{
 				BoneTransform = BoneTransform.GetRelativeTransform(FirstFrameBoneTransform);
 			}
+			else if (bUseRelativeToBone)
+			{
+				const FTransform RelativeToBoneTransform = UMotionExtractorUtilityLibrary::ExtractBoneTransform(Animation, BoneContainer, CompactPoseRelativeToBoneIndex, Time, bShouldUseComponentSpace);
+				BoneTransform = BoneTransform.GetRelativeTransform(RelativeToBoneTransform);
+			}
 
 			// Ignore first frame if we are extracting something that depends on the previous bone transform
 			if (SampleIndex > 1 || (MotionType != EMotionExtractor_MotionType::TranslationSpeed && MotionType != EMotionExtractor_MotionType::RotationSpeed))
@@ -116,12 +150,17 @@ void UMotionExtractorModifier::OnApply_Implementation(UAnimSequence* Animation)
 		Time = FMath::Clamp(static_cast<float>(SampleIndex) * SampleInterval, 0.f, AnimLength);
 		SampleIndex++;
 
-		FTransform BoneTransform = UMotionExtractorUtilityLibrary::ExtractBoneTransform(Animation, BoneContainer, CompactPoseBoneIndex, Time, bComponentSpace);
+		FTransform BoneTransform = UMotionExtractorUtilityLibrary::ExtractBoneTransform(Animation, BoneContainer, CompactPoseBoneIndex, Time, bShouldUseComponentSpace);
 
-		if (bRelativeToFirstFrame)
+		if(bRelativeToFirstFrame)
 		{
 			BoneTransform = BoneTransform.GetRelativeTransform(FirstFrameBoneTransform);
 		}
+		else if (bUseRelativeToBone)
+		{
+			const FTransform RelativeToBoneTransform = UMotionExtractorUtilityLibrary::ExtractBoneTransform(Animation, BoneContainer, CompactPoseRelativeToBoneIndex, Time, bShouldUseComponentSpace);
+			BoneTransform = BoneTransform.GetRelativeTransform(RelativeToBoneTransform);
+		}
 
 		// Ignore first frame if we are extracting something that depends on the previous bone transform
 		if (SampleIndex > 1 || (MotionType != EMotionExtractor_MotionType::TranslationSpeed && MotionType != EMotionExtractor_MotionType::RotationSpeed))
@@ -191,4 +230,16 @@ float UMotionExtractorModifier::GetDesiredValue(const FTransform& BoneTransform,
 	return Value;
 }
 
+void UMotionExtractorModifier::PostLoad()
+{
+	// If bComponentSpace wasn't its default value, set the space property and revert to default.
+	if (bComponentSpace_DEPRECATED == false)
+	{
+		bComponentSpace_DEPRECATED = true;
+		Space = EMotionExtractor_Space::LocalSpace;
+	}
+
+	Super::PostLoad();
+}
+
 #undef LOCTEXT_NAMESPACE
\ No newline at end of file
diff --git a/Engine/Plugins/Animation/AnimationModifierLibrary/Source/AnimationModifierLibrary/Public/MotionExtractorModifier.h b/Engine/Plugins/Animation/AnimationModifierLibrary/Source/AnimationModifierLibrary/Public/MotionExtractorModifier.h
--- a/Engine/Plugins/Animation/AnimationModifierLibrary/Source/AnimationModifierLibrary/Public/MotionExtractorModifier.h
+++ b/Engine/Plugins/Animation/AnimationModifierLibrary/Source/AnimationModifierLibrary/Public/MotionExtractorModifier.h
@@ -8,6 +8,7 @@
 enum class EMotionExtractor_Axis : uint8;
 enum class EMotionExtractor_MathOperation : uint8;
 enum class EMotionExtractor_MotionType : uint8;
+enum class EMotionExtractor_Space : uint8;
 
 /** Extracts motion from a bone in the animation and bakes it into a curve */
 UCLASS()
@@ -21,6 +22,10 @@ public:
 	UPROPERTY(EditAnywhere, BlueprintReadWrite, Category = Settings)
 	FName BoneName;
 
+	/** Bone which we'll use as the reference frame to generate this curve from. Space must be RelativeToBone to use this. */
+	UPROPERTY(EditAnywhere, BlueprintReadWrite, Category = Settings, meta=(EditCondition="Space == EMotionExtractor_Space::RelativeToBone"))
+	FName RelativeToBoneName;
+
 	/** Type of motion to extract */
 	UPROPERTY(EditAnywhere, BlueprintReadWrite, Category = Settings)
 	EMotionExtractor_MotionType MotionType;
@@ -39,9 +44,13 @@ public:
 	UPROPERTY(EditAnywhere, BlueprintReadWrite, Category = Settings)
 	bool bRelativeToFirstFrame;
 
-	/** Whether to extract the bone pose in component space or local space */
+	/** Reference frame/space to use when extracting the bone pose. */
 	UPROPERTY(EditAnywhere, BlueprintReadWrite, Category = Settings)
-	bool bComponentSpace;
+	EMotionExtractor_Space Space;
+
+	/** Whether to extract the bone pose in component space or local space. */
+	UPROPERTY()
+	bool bComponentSpace_DEPRECATED;
 
 	/** Whether to convert the final value to absolute (positive) */
 	UPROPERTY(EditAnywhere, BlueprintReadWrite, Category = Settings)
@@ -81,6 +90,8 @@ public:
 
 	/** Returns the desired value from the extracted poses */
 	float GetDesiredValue(const FTransform& BoneTransform, const FTransform& LastBoneTransform, float DeltaTime) const;
+
+	virtual void PostLoad() override;
 };
 
 #if UE_ENABLE_INCLUDE_ORDER_DEPRECATED_IN_5_2
diff --git a/Engine/Plugins/Animation/AnimationModifierLibrary/Source/AnimationModifierLibrary/Public/MotionExtractorTypes.h b/Engine/Plugins/Animation/AnimationModifierLibrary/Source/AnimationModifierLibrary/Public/MotionExtractorTypes.h
--- a/Engine/Plugins/Animation/AnimationModifierLibrary/Source/AnimationModifierLibrary/Public/MotionExtractorTypes.h
+++ b/Engine/Plugins/Animation/AnimationModifierLibrary/Source/AnimationModifierLibrary/Public/MotionExtractorTypes.h
@@ -39,6 +39,15 @@ enum class EMotionExtractor_Axis : uint8
 	XYZ
 };
 
+/** Reference frame/space to use when calculating motion */
+UENUM(BlueprintType)
+enum class EMotionExtractor_Space : uint8
+{
+	ComponentSpace,
+	LocalSpace,
+	RelativeToBone,
+};
+
 /** Math operations that can be applied to the extracted value before add it to the curve */
 UENUM(BlueprintType)
 enum class EMotionExtractor_MathOperation : uint8
