Fixed curves not being normalized from range 0 to 1 in FootAnimEventsModifier

[RN] minor

diff --git a/Engine/Plugins/Animation/AnimationModifierLibrary/Source/AnimationModifierLibrary/Private/FootstepAnimEventsModifier.cpp b/Engine/Plugins/Animation/AnimationModifierLibrary/Source/AnimationModifierLibrary/Private/FootstepAnimEventsModifier.cpp
--- a/Engine/Plugins/Animation/AnimationModifierLibrary/Source/AnimationModifierLibrary/Private/FootstepAnimEventsModifier.cpp
+++ b/Engine/Plugins/Animation/AnimationModifierLibrary/Source/AnimationModifierLibrary/Private/FootstepAnimEventsModifier.cpp
@@ -63,7 +63,10 @@ void UFootstepAnimEventsModifier::OnApply_Implementation(UAnimSequence* InAnimat
 				FFootSampleState & FootState = FootSampleStates[FootIndex];
 				
 				FootState.GroundLevel = FMath::Min(FootState.GroundLevel, UAnimPoseExtensions::GetBonePose(AnimPose, FootDef.FootBoneName, EAnimPoseSpaces::World).GetLocation().Z);
-				FootState.MaxFootSpeed = FMath::Max(FootState.MaxFootSpeed, FMath::Abs(ComputeBoneSpeed(AnimPose, FutureAnimPose, SampleStep, FootDef.FootBoneName)));
+
+				const float FootBoneSpeed = FMath::Abs(ComputeBoneSpeed(AnimPose, FutureAnimPose, SampleStep, FootDef.FootBoneName));
+				FootState.MaxFootSpeed = FMath::Max(FootState.MaxFootSpeed, FootBoneSpeed);
+				FootState.MinFootSpeed = FMath::Min(FootState.MinFootSpeed, FootBoneSpeed);
 			}
 		}
 
@@ -107,45 +110,76 @@ void UFootstepAnimEventsModifier::OnApply_Implementation(UAnimSequence* InAnimat
 					FootState.bIsFootBoneInGround = FMath::Abs(FootState.GroundLevel - FootBoneTransform.GetLocation().Z) <= GroundThreshold;
 
 					// Method: FootBoneSpeed
-					FootState.FootBoneSpeed = FMath::Abs(ComputeBoneSpeed(AnimPose, FutureAnimPose, SampleStep, FootDef.FootBoneName)) / FootState.MaxFootSpeed;
-
-					// Keep track of sample with lowest speed below threshold
-					if (FootState.FootBoneSpeed < SpeedThreshold)
 					{
-						if (FootState.FootBoneSpeed < FootState.MinFootSpeedBelowThreshold)
+						// Note that the stored foot bone speed is being mapped from range (MinFootSpeed, MaxFootSpeed) to (0, 1).
+						FootState.FootBoneSpeed = (FMath::Abs(ComputeBoneSpeed(AnimPose, FutureAnimPose, SampleStep, FootDef.FootBoneName)) - FootState.MinFootSpeed) / (FootState.MaxFootSpeed - FootState.MinFootSpeed);
+
+						// Keep track of sample with lowest speed below threshold
+						if (FootState.FootBoneSpeed < SpeedThreshold && SampleIndex > 1)
 						{
-							FootState.TimeAtMinFootSpeedBelowThreshold = SampleTime;
-							FootState.MinFootSpeedBelowThreshold = FootState.FootBoneSpeed;
+							if (FootState.FootBoneSpeed < FootState.MinFootSpeedBelowThreshold && FMath::Abs(FootState.MinFootSpeedBelowThreshold - FootState.FootBoneSpeed) >= 0.01f)
+							{
+								FootState.TimeAtMinFootSpeedBelowThreshold = SampleTime;
+								FootState.MinFootSpeedBelowThreshold = FootState.FootBoneSpeed;
+							}
 						}
 					}
+					
+					// Keep track if we are processing the last sample
+					FootState.bIsLast = SampleIndex == (SampleNum - 2);
 				}
-				
-				if (SampleIndex > 0)
+
+				// Place events when detection method passes
+				if (SampleIndex > 1)
 				{
 					// Generate sync markers
 					if (FootDef.bShouldGenerateSyncMarkers && CanWePlaceEventAtSample(FootState, FootDef.SyncMarkerDetectionTechnique))
 					{
-						float FinalSyncMarkerTime = FootDef.SyncMarkerDetectionTechnique == EDetectionTechnique::FootBoneSpeed ? FootState.TimeAtMinFootSpeedBelowThreshold : SampleTime;
-						UAnimationBlueprintLibrary::AddAnimationSyncMarker(InAnimation, FootDef.SyncMarkerName, FinalSyncMarkerTime, FootDef.SyncMarkerTrackName);
+						if (!FootState.bSkipNextSyncMarker)
+						{
+							const float FinalSyncMarkerTime = FootDef.SyncMarkerDetectionTechnique == EDetectionTechnique::FootBoneSpeed ? FootState.TimeAtMinFootSpeedBelowThreshold : SampleTime;
+							UAnimationBlueprintLibrary::AddAnimationSyncMarker(InAnimation, FootDef.SyncMarkerName, FinalSyncMarkerTime, FootDef.SyncMarkerTrackName);
+						}
+						else
+						{
+							FootState.bSkipNextSyncMarker = false;
+						}
 					}
 
 					// Generate foot step fx notifies
 					if (FootDef.bShouldGenerateNotifies && CanWePlaceEventAtSample(FootState, FootDef.FootstepNotifyDetectionTechnique))
 					{
-						float FinalAnimNotifyTime = FootDef.FootstepNotifyDetectionTechnique == EDetectionTechnique::FootBoneSpeed ? FootState.TimeAtMinFootSpeedBelowThreshold : SampleTime;
-						UAnimationBlueprintLibrary::AddAnimationNotifyEvent(InAnimation, FootDef.FootstepNotifyTrackName, FinalAnimNotifyTime, FootDef.FootstepNotify);
+						if (!FootState.bSkipNextNotify)
+						{
+							const float FinalAnimNotifyTime = FootDef.FootstepNotifyDetectionTechnique == EDetectionTechnique::FootBoneSpeed ? FootState.TimeAtMinFootSpeedBelowThreshold : SampleTime;
+							UAnimationBlueprintLibrary::AddAnimationNotifyEvent(InAnimation, FootDef.FootstepNotifyTrackName, FinalAnimNotifyTime, FootDef.FootstepNotify);
+						}
+						else
+						{
+							FootState.bSkipNextNotify = false;
+						}
 					}
 				}
 				
 				// Update foot state
 				{
-					const bool bDidWePlaceAnyEventsUsingFootBoneSpeed = CanWePlaceEventAtSample(FootState, EDetectionTechnique::FootBoneSpeed);
+					if (SampleIndex > 1)
+					{
+						const bool bDidEventPlacementTestPassUsingFootBoneSpeed = CanWePlaceEventAtSample(FootState, EDetectionTechnique::FootBoneSpeed);
 
-					// Reset minimum values if we are above speed threshold
-					if (bDidWePlaceAnyEventsUsingFootBoneSpeed)
+						// Reset minimum values if we are above speed threshold
+						if (bDidEventPlacementTestPassUsingFootBoneSpeed)
+						{
+							// Reset minimum values if we are above speed threshold
+							FootState.TimeAtMinFootSpeedBelowThreshold = MAX_FLT;
+							FootState.MinFootSpeedBelowThreshold = MAX_FLT;
+						}
+					}
+					else
 					{
-						FootState.TimeAtMinFootSpeedBelowThreshold = MAX_FLT;
-						FootState.MinFootSpeedBelowThreshold = MAX_FLT;
+						// Skip anim events if foot started below speed threshold
+						FootState.bSkipNextNotify = FootState.FootBoneSpeed < SpeedThreshold && FootDef.bShouldSkipNotifyIfFootBoneSpeedStartsBelowThreshold;
+						FootState.bSkipNextSyncMarker = FootState.FootBoneSpeed < SpeedThreshold && FootDef.bShouldSkipSyncMarkerIfFootBoneSpeedStartsBelowThreshold;
 					}
 					
 					// Keep track of foot info
@@ -272,7 +306,7 @@ bool UFootstepAnimEventsModifier::CanWePlaceEventAtSample(const FFootSampleState
 	{
 		case EDetectionTechnique::PassThroughReferenceBone: return InFootSampleState.RefBoneTranslationDotRefBoneToFootBoneVec > 0.0f && InFootSampleState.PrevRefBoneTranslationDotRefBoneToFootBoneVec < 0.0f;
 		case EDetectionTechnique::FootBoneReachesGround: return !InFootSampleState.bWasFootBoneInGround && InFootSampleState.bIsFootBoneInGround;
-		case EDetectionTechnique::FootBoneSpeed: return InFootSampleState.PrevFootBoneSpeed < SpeedThreshold && InFootSampleState.FootBoneSpeed >= SpeedThreshold;
+		case EDetectionTechnique::FootBoneSpeed: return (InFootSampleState.PrevFootBoneSpeed < SpeedThreshold && InFootSampleState.FootBoneSpeed >= SpeedThreshold) || (InFootSampleState.bIsLast && InFootSampleState.FootBoneSpeed < SpeedThreshold && !FMath::IsNearlyEqual(InFootSampleState.MinFootSpeedBelowThreshold, MAX_FLT));
 		default: return false;
 	}
 }
diff --git a/Engine/Plugins/Animation/AnimationModifierLibrary/Source/AnimationModifierLibrary/Public/FootstepAnimEventsModifier.h b/Engine/Plugins/Animation/AnimationModifierLibrary/Source/AnimationModifierLibrary/Public/FootstepAnimEventsModifier.h
--- a/Engine/Plugins/Animation/AnimationModifierLibrary/Source/AnimationModifierLibrary/Public/FootstepAnimEventsModifier.h
+++ b/Engine/Plugins/Animation/AnimationModifierLibrary/Source/AnimationModifierLibrary/Public/FootstepAnimEventsModifier.h
@@ -58,6 +58,9 @@ struct FFootDefinition
 
 	UPROPERTY(EditAnywhere, Category = "Sync Markers", meta=(EditCondition="bShouldGenerateSyncMarkers"))
 	EDetectionTechnique SyncMarkerDetectionTechnique = EDetectionTechnique::PassThroughReferenceBone;
+
+	UPROPERTY(EditAnywhere, Category = "Sync Markers", meta=(EditCondition = "SyncMarkerDetectionTechnique == EDetectionTechnique::FootBoneSpeed", EditConditionHides))
+	bool bShouldSkipSyncMarkerIfFootBoneSpeedStartsBelowThreshold = true;
 	
 	UPROPERTY(EditAnywhere, Category = "Notifies")
 	bool bShouldGenerateNotifies = false;
@@ -69,7 +72,10 @@ struct FFootDefinition
 	TSubclassOf<UAnimNotify> FootstepNotify = nullptr;
 
 	UPROPERTY(EditAnywhere, Category = "Notifies", meta=(EditCondition="bShouldGenerateNotifies"))
-	EDetectionTechnique FootstepNotifyDetectionTechnique = EDetectionTechnique::FootBoneReachesGround;
+	EDetectionTechnique FootstepNotifyDetectionTechnique = EDetectionTechnique::FootBoneSpeed;
+
+	UPROPERTY(EditAnywhere, Category = "Notifies", meta=(EditCondition = "FootstepNotifyDetectionTechnique == EDetectionTechnique::FootBoneSpeed", EditConditionHides))
+	bool bShouldSkipNotifyIfFootBoneSpeedStartsBelowThreshold = true;
 };
 
 /** Generates animation notifies and/or sync markers for any specified bone(s) */
@@ -128,8 +134,12 @@ private:
 		float FootBoneSpeed = 0.0f;
 		float PrevFootBoneSpeed = 0.0f;
 		float MaxFootSpeed = -MAX_FLT;
+		float MinFootSpeed = MAX_FLT;
 		float MinFootSpeedBelowThreshold = MAX_FLT;
 		float TimeAtMinFootSpeedBelowThreshold = MAX_FLT;
+		bool bIsLast = false;
+		bool bSkipNextSyncMarker = false;
+		bool bSkipNextNotify = false;
 	};
 	
 	/** Keep track of to be generated tracks during modifier application */
