Created a Mirror AnimationModifier that mirrors all bones in a sequence based on a mirror data table. This changes is intended for UE5.

diff --git a/Engine/Plugins/Animation/AnimationModifierLibrary/Source/AnimationModifierLibrary/Private/MirrorModifier.cpp b/Engine/Plugins/Animation/AnimationModifierLibrary/Source/AnimationModifierLibrary/Private/MirrorModifier.cpp
new file mode 100644
--- /dev/null
+++ b/Engine/Plugins/Animation/AnimationModifierLibrary/Source/AnimationModifierLibrary/Private/MirrorModifier.cpp
@@ -0,0 +1,138 @@
+// Copyright Epic Games, Inc. All Rights Reserved.
+
+#include "MirrorModifier.h"
+#include "Animation/AnimData/AnimDataModel.h"
+#include "Animation/AnimSequence.h"
+#include "AnimationRuntime.h"
+#include "Animation/AnimData/IAnimationDataController.h"
+#include "EngineLogs.h"
+
+#define LOCTEXT_NAMESPACE "MirrorModifier"
+
+UMirrorModifier::UMirrorModifier()
+	:Super()
+{
+}
+
+void UMirrorModifier::OnApply_Implementation(UAnimSequence* Animation)
+{
+	if (Animation == nullptr)
+	{
+		UE_LOG(LogAnimation, Error, TEXT("MirrorModifier failed. Reason: Invalid Animation"));
+		return;
+	}
+
+	IAnimationDataController& Controller = Animation->GetController();
+	const IAnimationDataModel* Model = Animation->GetDataModel();
+
+	if (Model == nullptr)
+	{
+		UE_LOG(LogAnimation, Error, TEXT("MirrorModifier failed. Reason: Invalid Data Model. Animation: %s"), *GetNameSafe(Animation));
+		return;
+	}
+
+	if (!MirrorDataTable)
+	{
+		UE_LOG(LogAnimation, Error, TEXT("MirrorModifier failed. Reason: Invalid Mirror Data table"));
+		return;
+	}
+
+	// Compact pose format of Mirror Bone Map
+	TCustomBoneIndexArray<FCompactPoseBoneIndex, FCompactPoseBoneIndex> CompactPoseMirrorBones;
+	FMemMark Mark(FMemStack::Get());
+
+	// Pre-calculated component space of reference pose, which allows mirror to work with any joint orient 
+	TCustomBoneIndexArray<FQuat, FCompactPoseBoneIndex> ComponentSpaceRefRotations;
+	FBoneContainer BoneContainer;
+	InitializeBoneContainer(Animation, BoneContainer);
+	MirrorDataTable->FillCompactPoseAndComponentRefRotations(
+		BoneContainer,
+		CompactPoseMirrorBones,
+		ComponentSpaceRefRotations);
+
+	// Start editing animation data
+	constexpr bool bShouldTransact = false;	Controller.OpenBracket(LOCTEXT("MirrorModifier_Bracket", "Updating bones"), bShouldTransact);
+
+	const int32 NumKeys = Model->GetNumberOfKeys();
+	for (int32 AnimKey = 0; AnimKey < NumKeys; AnimKey++)
+	{
+			FAnimPose AnimPose;
+			UAnimPoseExtensions::GetAnimPoseAtFrame(Animation, AnimKey, FAnimPoseEvaluationOptions(), AnimPose);
+			FMemMark KeyMark(FMemStack::Get());
+			FCompactPose OutPose;
+			OutPose.SetBoneContainer(&BoneContainer);
+			UAnimPoseExtensions::GetCompactPose(AnimPose, OutPose); 
+
+			FAnimationRuntime::MirrorPose(OutPose, MirrorDataTable->MirrorAxis, CompactPoseMirrorBones, ComponentSpaceRefRotations);
+
+			for (int32 BoneIndex = 0; BoneIndex < BoneContainer.GetNumBones(); ++BoneIndex)
+			{
+				FName SourceBoneName = BoneContainer.GetReferenceSkeleton().GetBoneName(BoneIndex);
+				FSkeletonPoseBoneIndex SkeletonBoneIndex(BoneIndex);
+				FCompactPoseBoneIndex CompactIndex = BoneContainer.GetCompactPoseIndexFromSkeletonPoseIndex(SkeletonBoneIndex);
+				FTransform BonePoseTargetLocal = OutPose[CompactIndex];
+				const FInt32Range KeyRangeToSet(AnimKey, AnimKey + 1);
+				if (Model->IsValidBoneTrackName(SourceBoneName))
+				{
+					Controller.UpdateBoneTrackKeys(SourceBoneName, KeyRangeToSet, { BonePoseTargetLocal.GetLocation() }, { BonePoseTargetLocal.GetRotation() }, { BonePoseTargetLocal.GetScale3D() });
+				}
+
+			}	
+	}
+
+	if (bUpdateSyncMarkers)
+	{
+		for(FAnimSyncMarker& Marker : Animation->AuthoredSyncMarkers)
+		{
+			const FName* MirroredName = MirrorDataTable->SyncToMirrorSyncMap.Find(Marker.MarkerName);
+			if (MirroredName)
+			{
+				Marker.MarkerName = *MirroredName;
+			}
+		}
+		Animation->RefreshSyncMarkerDataFromAuthored(); 
+	}
+
+	if (bUpdateNotifies)
+	{
+		for(FAnimNotifyEvent& NotifyEvent : Animation->Notifies)
+		{
+			const FName* MirroredName = MirrorDataTable->AnimNotifyToMirrorAnimNotifyMap.Find(NotifyEvent.NotifyName);
+			if (MirroredName)
+			{
+				NotifyEvent.NotifyName = *MirroredName;
+			}
+		}
+		Animation->RefreshCacheData();
+	}
+
+	// Done editing animation data
+	Controller.CloseBracket(bShouldTransact);
+}
+
+void UMirrorModifier::OnRevert_Implementation(UAnimSequence* Animation)
+{
+	//  mirroring will undo the previous mirror 
+	OnApply_Implementation(Animation);
+}
+
+void UMirrorModifier::InitializeBoneContainer(const UAnimSequenceBase* AnimationSequenceBase, FBoneContainer& OutContainer) const
+{
+	// asset to use for retarget proportions (can be either USkeletalMesh or USkeleton)
+	UObject* AssetToUse = CastChecked<UObject>(AnimationSequenceBase->GetSkeleton());
+	int32 NumRequiredBones = AnimationSequenceBase->GetSkeleton()->GetReferenceSkeleton().GetNum();
+
+	TArray<FBoneIndexType> RequiredBoneIndexArray;
+	RequiredBoneIndexArray.AddUninitialized(NumRequiredBones);
+	for (int32 BoneIndex = 0; BoneIndex < RequiredBoneIndexArray.Num(); ++BoneIndex)
+	{
+		RequiredBoneIndexArray[BoneIndex] = static_cast<FBoneIndexType>(BoneIndex);
+	}
+
+	OutContainer.InitializeTo(RequiredBoneIndexArray, UE::Anim::FCurveFilterSettings(UE::Anim::ECurveFilterMode::DisallowAll), *AssetToUse);
+	OutContainer.SetUseRAWData(true);
+	OutContainer.SetUseSourceData(true);
+	OutContainer.SetDisableRetargeting(false);
+}
+
+#undef LOCTEXT_NAMESPACE
diff --git a/Engine/Plugins/Animation/AnimationModifierLibrary/Source/AnimationModifierLibrary/Public/MirrorModifier.h b/Engine/Plugins/Animation/AnimationModifierLibrary/Source/AnimationModifierLibrary/Public/MirrorModifier.h
new file mode 100644
--- /dev/null
+++ b/Engine/Plugins/Animation/AnimationModifierLibrary/Source/AnimationModifierLibrary/Public/MirrorModifier.h
@@ -0,0 +1,37 @@
+// Copyright Epic Games, Inc. All Rights Reserved.
+
+#pragma once
+
+#include "AnimationModifier.h"
+#include "AnimPose.h"
+#include "BoneContainer.h"
+#include "Animation/MirrorDataTable.h"
+#include "MirrorModifier.generated.h"
+
+
+/** Animation Modifier to mirror an animation using a mirror data table */
+UCLASS()
+class UMirrorModifier : public UAnimationModifier
+{
+	GENERATED_BODY()
+
+public:
+
+	//MirrorDataTable that will be used for mirroring the aniamtion
+	UPROPERTY(EditAnywhere, Category = "Settings")
+	TObjectPtr<UMirrorDataTable> MirrorDataTable;
+
+	// If sync markers should be mirrored in the animation 
+	UPROPERTY(EditAnywhere, Category = "Settings")
+	bool  bUpdateSyncMarkers = true;
+
+	// If notifies should be mirrored in the animation 
+	UPROPERTY(EditAnywhere, Category = "Settings")
+	bool  bUpdateNotifies = true;
+	
+	UMirrorModifier();
+	virtual void OnApply_Implementation(UAnimSequence* Animation) override;
+	virtual void OnRevert_Implementation(UAnimSequence* Animation) override;
+protected:
+	void InitializeBoneContainer(const UAnimSequenceBase* AnimationSequenceBase, FBoneContainer& OutContainer) const;
+};
diff --git a/Engine/Source/Editor/AnimationBlueprintLibrary/Private/AnimPose.cpp b/Engine/Source/Editor/AnimationBlueprintLibrary/Private/AnimPose.cpp
--- a/Engine/Source/Editor/AnimationBlueprintLibrary/Private/AnimPose.cpp
+++ b/Engine/Source/Editor/AnimationBlueprintLibrary/Private/AnimPose.cpp
@@ -706,3 +706,9 @@ void UAnimPoseExtensions::GetAnimPoseAtTimeIntervals(const UAnimSequenceBase* An
 		UE_LOG(LogAnimationPoseScripting, Warning, TEXT("Invalid Animation Sequence supplied for GetBonePosesForTime"));
 	}
 }
+
+void UAnimPoseExtensions::GetCompactPose(const FAnimPose& Pose, FCompactPose& InOutCompactPose)
+{
+	Pose.GetPose(InOutCompactPose); 
+}
+
diff --git a/Engine/Source/Editor/AnimationBlueprintLibrary/Public/AnimPose.h b/Engine/Source/Editor/AnimationBlueprintLibrary/Public/AnimPose.h
--- a/Engine/Source/Editor/AnimationBlueprintLibrary/Public/AnimPose.h
+++ b/Engine/Source/Editor/AnimationBlueprintLibrary/Public/AnimPose.h
@@ -346,4 +346,13 @@ public:
     */
     UFUNCTION(BlueprintPure, meta = (ScriptMethod), Category = "Animation|Pose")
     static float GetCurveWeight(UPARAM(ref) const FAnimPose& Pose, const FName& CurveName);
+
+
+	/**
+    * Returns a CompactPose from the given FAnimPose.  This function should not be exposed to blueprint 
+    *
+    * @param	Pose		Anim Pose to retrieve the value from
+    * @param	InOutCompactPose	CompactPose from the Pose
+    */
+	static void GetCompactPose(const FAnimPose& Pose, FCompactPose& InOutCompactPose);
 };
\ No newline at end of file
