[Animation] Prevent some modifier to go out off memory and make them execute a lot faster. There was a missing bShoudlTransact /*false*/ that was forcing to store thousand of Transform in the transaction buffer. Those transaction are unnecessary since there was already a higher level transaction that have a copy of the anim sequence.

There is also a fix for an asynchronous crash because there was no wait before an anim sequence compression is done before applying a modifier.

There is also now some slow task progress dialog showing what the editor is doing since applying modifier can be a long task.

#rnx

diff --git a/Engine/Plugins/Animation/AnimationModifierLibrary/Source/AnimationModifierLibrary/Private/CopyBonesModifier.cpp b/Engine/Plugins/Animation/AnimationModifierLibrary/Source/AnimationModifierLibrary/Private/CopyBonesModifier.cpp
--- a/Engine/Plugins/Animation/AnimationModifierLibrary/Source/AnimationModifierLibrary/Private/CopyBonesModifier.cpp
+++ b/Engine/Plugins/Animation/AnimationModifierLibrary/Source/AnimationModifierLibrary/Private/CopyBonesModifier.cpp
@@ -5,6 +5,7 @@
 #include "Animation/AnimSequence.h"
 #include "Animation/AnimData/IAnimationDataController.h"
 #include "EngineLogs.h"
+#include "Misc/ScopedSlowTask.h"
 
 #define LOCTEXT_NAMESPACE "CopyBonesModifier"
 
@@ -76,8 +77,11 @@ void UCopyBonesModifier::OnApply_Implementation(UAnimSequence* Animation)
 
 	// Get the transform of all the source bones in the desired space
 	const int32 NumKeys = Model->GetNumberOfKeys();
+	FScopedSlowTask SlowTask(static_cast<float>(NumKeys));
+	SlowTask.MakeDialog();
 	for (int32 AnimKey = 0; AnimKey < NumKeys; AnimKey++)
 	{
+		SlowTask.EnterProgressFrame();
 		for (FCopyBoneData& Data : CopyBoneDataContainer)
 		{
 			FAnimPose AnimPose;
@@ -90,7 +94,7 @@ void UCopyBonesModifier::OnApply_Implementation(UAnimSequence* Animation)
 			FTransform BonePoseTargetLocal = UAnimPoseExtensions::GetBonePose(AnimPose, Data.TargetBoneName, EAnimPoseSpaces::Local);
 
 			const FInt32Range KeyRangeToSet(AnimKey, AnimKey + 1);
-			Controller.UpdateBoneTrackKeys(Data.TargetBoneName, KeyRangeToSet, { BonePoseTargetLocal.GetLocation() }, { BonePoseTargetLocal.GetRotation() }, { BonePoseTargetLocal.GetScale3D() });
+			Controller.UpdateBoneTrackKeys(Data.TargetBoneName, KeyRangeToSet, { BonePoseTargetLocal.GetLocation() }, { BonePoseTargetLocal.GetRotation() }, { BonePoseTargetLocal.GetScale3D() }, bShouldTransact);
 		}
 	}
 
diff --git a/Engine/Plugins/Animation/AnimationModifierLibrary/Source/AnimationModifierLibrary/Private/EncodeRootBoneModifier.cpp b/Engine/Plugins/Animation/AnimationModifierLibrary/Source/AnimationModifierLibrary/Private/EncodeRootBoneModifier.cpp
--- a/Engine/Plugins/Animation/AnimationModifierLibrary/Source/AnimationModifierLibrary/Private/EncodeRootBoneModifier.cpp
+++ b/Engine/Plugins/Animation/AnimationModifierLibrary/Source/AnimationModifierLibrary/Private/EncodeRootBoneModifier.cpp
@@ -6,6 +6,7 @@
 
 #include "EncodeRootBoneModifier.h"
 #include "Animation/AnimSequence.h"
+#include "Misc/ScopedSlowTask.h"
 #include "PoseExtractorUtilityLibrary.h"
 
 #define LOCTEXT_NAMESPACE "EncodeRootBoneModifier"
@@ -59,11 +60,13 @@ void UEncodeRootBoneModifier::OnApply_Implementation(UAnimSequence* Animation)
 	// Start editing animation data
 	const bool bShouldTransact = false;
 	Animation->GetController().OpenBracket(LOCTEXT("EncodeRootBoneModifier_Bracket", "Encoding root bone"), bShouldTransact);
-
 	// For each key in the animation
 	const int32 NumAnimKeys = Animation->GetDataModel()->GetNumberOfKeys();
+	FScopedSlowTask SlowTask(static_cast<float>(NumAnimKeys));
+	SlowTask.MakeDialog();
 	for (int32 AnimKey = 0; AnimKey < NumAnimKeys; AnimKey++)
 	{
+		SlowTask.EnterProgressFrame();
 		const FInt32Range KeyRangeToSet(AnimKey, AnimKey + 1);
 
 		const FTransform RootTransformOriginal = Extractor.GetGlobalBoneTransform(AnimKey, 0);
@@ -106,13 +109,13 @@ void UEncodeRootBoneModifier::OnApply_Implementation(UAnimSequence* Animation)
 			RootTransformNew.SetRotation(Rotator.Quaternion());
 		}
 
-		Animation->GetController().UpdateBoneTrackKeys(Animation->GetSkeleton()->GetReferenceSkeleton().GetBoneName(0), KeyRangeToSet, { RootTransformNew.GetLocation() }, { RootTransformNew.GetRotation() }, { RootTransformNew.GetScale3D() });
+		Animation->GetController().UpdateBoneTrackKeys(Animation->GetSkeleton()->GetReferenceSkeleton().GetBoneName(0), KeyRangeToSet, { RootTransformNew.GetLocation() }, { RootTransformNew.GetRotation() }, { RootTransformNew.GetScale3D() }, bShouldTransact);
 
 		// Now the mesh is facing the wrong axis. Update direct children of the root with the local space transform that puts them back to where they were originally
 		for (const FName& ChildBoneName : RootBoneChildBones)
 		{
 			const FTransform ChildBoneNewTransform = Extractor.GetGlobalBoneTransform(AnimKey, ChildBoneName).GetRelativeTransform(RootTransformNew);
-			Animation->GetController().UpdateBoneTrackKeys(ChildBoneName, KeyRangeToSet, { ChildBoneNewTransform.GetLocation() }, { ChildBoneNewTransform.GetRotation() }, { ChildBoneNewTransform.GetScale3D() });
+			Animation->GetController().UpdateBoneTrackKeys(ChildBoneName, KeyRangeToSet, { ChildBoneNewTransform.GetLocation() }, { ChildBoneNewTransform.GetRotation() }, { ChildBoneNewTransform.GetScale3D() }, bShouldTransact);
 		}
 	}
 
diff --git a/Engine/Plugins/Animation/AnimationModifierLibrary/Source/AnimationModifierLibrary/Private/ZeroOutRootBoneModifier.cpp b/Engine/Plugins/Animation/AnimationModifierLibrary/Source/AnimationModifierLibrary/Private/ZeroOutRootBoneModifier.cpp
--- a/Engine/Plugins/Animation/AnimationModifierLibrary/Source/AnimationModifierLibrary/Private/ZeroOutRootBoneModifier.cpp
+++ b/Engine/Plugins/Animation/AnimationModifierLibrary/Source/AnimationModifierLibrary/Private/ZeroOutRootBoneModifier.cpp
@@ -8,6 +8,7 @@
 #include "AnimationUtils.h"
 #include "Animation/Skeleton.h"
 #include "EngineLogs.h"
+#include "Misc/ScopedSlowTask.h"
 
 #define LOCTEXT_NAMESPACE "ZeroOutRootBoneModifier"
 
@@ -60,8 +61,11 @@ void UZeroOutRootBoneModifier::OnApply_Implementation(UAnimSequence* Animation)
 	PreModifiedRootTransforms.Reserve(NumKeys);
 	int32 KeyWhereMotionStarts = -1;
 	int32 KeyWhereMotionEnds = 0;
+	FScopedSlowTask SlowTask(static_cast<float>(NumKeys));
+	SlowTask.MakeDialog();
 	for (int32 AnimKey = 0; AnimKey < NumKeys; AnimKey++)
 	{
+		SlowTask.EnterProgressFrame();
 		const FTransform RootTransform = Model->EvaluateBoneTrackTransform(RootBoneName, AnimKey, EAnimInterpolationType::Step);
 		PreModifiedRootTransforms.Add(RootTransform);
 
@@ -88,7 +92,7 @@ void UZeroOutRootBoneModifier::OnApply_Implementation(UAnimSequence* Animation)
 		const FInt32Range KeyRangeToSet(AnimKey, AnimKey + 1);
 
 		FTransform NewRootTransform = PreModifiedRootTransforms[AnimKey] * PreModifiedRootTransforms[0].Inverse();
-		Controller.UpdateBoneTrackKeys(RootBoneName, KeyRangeToSet, { NewRootTransform.GetLocation() }, { NewRootTransform.GetRotation() }, { NewRootTransform.GetScale3D() });
+		Controller.UpdateBoneTrackKeys(RootBoneName, KeyRangeToSet, { NewRootTransform.GetLocation() }, { NewRootTransform.GetRotation() }, { NewRootTransform.GetScale3D() }, bShouldTransact);
 	}
 
 	// Done editing animation data
diff --git a/Engine/Source/Editor/AnimationModifiers/Private/AnimationModifier.cpp b/Engine/Source/Editor/AnimationModifiers/Private/AnimationModifier.cpp
--- a/Engine/Source/Editor/AnimationModifiers/Private/AnimationModifier.cpp
+++ b/Engine/Source/Editor/AnimationModifiers/Private/AnimationModifier.cpp
@@ -42,11 +42,18 @@ UAnimationModifier::UAnimationModifier()
 {
 }
 
-void UAnimationModifier::ApplyToAnimationSequence(class UAnimSequence* AnimSequence) const
+void UAnimationModifier::ApplyToAnimationSequence(UAnimSequence* AnimSequence) const
 {
 	FEditorScriptExecutionGuard ScriptGuard;
 
 	checkf(AnimSequence, TEXT("Invalid Animation Sequence supplied"));
+
+	//TODO: We want to use UE_COMPRESSED_DATA_READ_SCOPE instead of a wait. But the scope is not working since the animation modifier will modify the anim sequence
+	//and will deadlock when closing the animation controller bracket.
+	//UE_COMPRESSED_DATA_READ_SCOPE(AnimSequence);
+	//Make sure Anim sequence compression is done before starting a new one.
+	AnimSequence->WaitOnExistingCompression();
+
 	USkeleton* Skeleton = AnimSequence->GetSkeleton();
 	checkf(Skeleton, TEXT("Invalid Skeleton for supplied Animation Sequence"));
 
diff --git a/Engine/Source/Editor/AnimationModifiers/Private/AnimationModifiersModule.cpp b/Engine/Source/Editor/AnimationModifiers/Private/AnimationModifiersModule.cpp
--- a/Engine/Source/Editor/AnimationModifiers/Private/AnimationModifiersModule.cpp
+++ b/Engine/Source/Editor/AnimationModifiers/Private/AnimationModifiersModule.cpp
@@ -21,6 +21,7 @@
 #include "Internationalization/Internationalization.h"
 #include "Math/Vector2D.h"
 #include "Misc/CoreDelegates.h"
+#include "Misc/ScopedSlowTask.h"
 #include "Modules/ModuleManager.h"
 #include "PropertyEditorDelegates.h"
 #include "PropertyEditorModule.h"
@@ -368,23 +369,30 @@ void FAnimationModifiersModule::ApplyAnimationModifiers(const TArray<UAnimSequen
 	// Iterate over each Animation Sequence and all of its contained modifiers, applying each one
 	UE::Anim::FApplyModifiersScope Scope;
 	TArray<UAnimationModifiersAssetUserData*> AssetUserData;
+	FScopedSlowTask SlowTaskSequences(InSequences.Num(), LOCTEXT("SlowTaskApplyingModifier_Sequence_Start", "Applying Animation Modifier(s) to Animation Sequence(s)..."));
+	SlowTaskSequences.MakeDialog();
 	for (UAnimSequence* AnimationSequence : InSequences)
 	{
+		FText AnimSequenceName = FText::FromString(IsValid(AnimationSequence) ? AnimationSequence->GetName() : TEXT(""));
+		SlowTaskSequences.EnterProgressFrame(1.0f, FText::Format(LOCTEXT("SlowTaskApplyingModifier_Sequence_Update", "Applying Animation Modifier(s) to Animation Sequence {0}..."), AnimSequenceName));
 		if (AnimationSequence)
 		{
 			UAnimationModifiersAssetUserData* UserData = AnimationSequence->GetAssetUserData<UAnimationModifiersAssetUserData>();
 			if (UserData)
 			{
 				const TArray<UAnimationModifier*>& ModifierInstances = UserData->GetAnimationModifierInstances();
+				FScopedSlowTask SlowTaskModifiers(ModifierInstances.Num(), LOCTEXT("SlowTaskApplyingModifier_Modifier_Start", "Applying Animation Modifiers..."));
+				SlowTaskModifiers.MakeDialog();
 				for (UAnimationModifier* Modifier : ModifierInstances)
 				{
+					SlowTaskModifiers.EnterProgressFrame(1.0f, FText::Format(LOCTEXT("SlowTaskApplyingModifier_Modifier_Update", "Applying Animation Modifier {0}..."), Modifier->GetClass()->GetDisplayNameText()));
 					if (bForceApply || !Modifier->IsLatestRevisionApplied(AnimationSequence))
 					{
 						Modifier->ApplyToAnimationSequence(AnimationSequence);
 					}
 				}
 			}
-		}		
+		}
 	}
 }
 
diff --git a/Engine/Source/Editor/AnimationModifiers/Private/SAnimationModifiersTab.cpp b/Engine/Source/Editor/AnimationModifiers/Private/SAnimationModifiersTab.cpp
--- a/Engine/Source/Editor/AnimationModifiers/Private/SAnimationModifiersTab.cpp
+++ b/Engine/Source/Editor/AnimationModifiers/Private/SAnimationModifiersTab.cpp
@@ -30,6 +30,7 @@
 #include "Misc/AssertionMacros.h"
 #include "Misc/Attribute.h"
 #include "Misc/MessageDialog.h"
+#include "Misc/ScopedSlowTask.h"
 #include "Modules/ModuleManager.h"
 #include "PropertyEditorModule.h"
 #include "SModifierListview.h"
@@ -501,13 +502,19 @@ void SAnimationModifiersTab::ApplyModifiers(const TArray<UAnimationModifier*>& M
 	if (bApply)
 	{
 		UE::Anim::FApplyModifiersScope Scope;
-
+		FScopedSlowTask SlowTaskModifier(Modifiers.Num(), LOCTEXT("SlowTaskApplyingModifier_Modifier_Start", "Applying Animation Modifier(s) to Animation Sequence(s)...") );
+		SlowTaskModifier.MakeDialog();
 		for (UAnimationModifier* Instance : Modifiers)
 		{
 			checkf(Instance, TEXT("Invalid modifier instance"));
+			SlowTaskModifier.EnterProgressFrame(1.0f, FText::Format(LOCTEXT("SlowTaskApplyingModifier_Modifier_Update", "Applying Animation Modifier {0}..."), Instance->GetClass()->GetDisplayNameText()));
 			bool AppliedOnAny = false;
+			FScopedSlowTask SlowTaskSequences(AnimSequences.Num(), LOCTEXT("SlowTaskApplyingModifier_Sequence_Start", "Applying Animation Modifier to Animation Sequence(s)..."));
+			SlowTaskSequences.MakeDialog();
 			for (UAnimSequence* AnimSequence : AnimSequences)
 			{
+				FText AnimSequenceName = FText::FromString(IsValid(AnimSequence) ? AnimSequence->GetName() : TEXT(""));
+				SlowTaskSequences.EnterProgressFrame(1.0f, FText::Format(LOCTEXT("SlowTaskApplyingModifier_Sequence_Update", "Applying Animation Modifier to Animation Sequence {0}..."), AnimSequenceName));
 				ensure(!(Skeleton != nullptr) || AnimSequence->GetSkeleton() == Skeleton);
 				Instance->ApplyToAnimationSequence(AnimSequence);
 				AppliedOnAny = AppliedOnAny || Instance->CanRevert(AnimSequence);
