Fixed an issue with the mirror modifier stomping future keys when using the UpdateBoneTrackKeys APIs on a single frame

diff --git a/Engine/Plugins/Animation/AnimationModifierLibrary/Source/AnimationModifierLibrary/Private/MirrorModifier.cpp b/Engine/Plugins/Animation/AnimationModifierLibrary/Source/AnimationModifierLibrary/Private/MirrorModifier.cpp
--- a/Engine/Plugins/Animation/AnimationModifierLibrary/Source/AnimationModifierLibrary/Private/MirrorModifier.cpp
+++ b/Engine/Plugins/Animation/AnimationModifierLibrary/Source/AnimationModifierLibrary/Private/MirrorModifier.cpp
@@ -52,32 +52,41 @@ void UMirrorModifier::OnApply_Implementation(UAnimSequence* Animation)
 
 	// Start editing animation data
 	constexpr bool bShouldTransact = false;	Controller.OpenBracket(LOCTEXT("MirrorModifier_Bracket", "Updating bones"), bShouldTransact);
-
 	const int32 NumKeys = Model->GetNumberOfKeys();
+
+	TArray<FAnimPose>  StoredPoses;
+	StoredPoses.Reserve(NumKeys);
+
 	for (int32 AnimKey = 0; AnimKey < NumKeys; AnimKey++)
 	{
-			FAnimPose AnimPose;
-			UAnimPoseExtensions::GetAnimPoseAtFrame(Animation, AnimKey, FAnimPoseEvaluationOptions(), AnimPose);
-			FMemMark KeyMark(FMemStack::Get());
-			FCompactPose OutPose;
-			OutPose.SetBoneContainer(&BoneContainer);
-			UAnimPoseExtensions::GetCompactPose(AnimPose, OutPose); 
+		FAnimPose AnimPose;
+		UAnimPoseExtensions::GetAnimPoseAtFrame(Animation, AnimKey, FAnimPoseEvaluationOptions(), AnimPose);
+		StoredPoses.Emplace(MoveTemp(AnimPose));
+	}
+	
+	for (int32 AnimKey = 0; AnimKey < NumKeys; AnimKey++)
+	{
+		FAnimPose& AnimPose = StoredPoses[AnimKey];
+		FMemMark KeyMark(FMemStack::Get());
+		FCompactPose OutPose;
+		OutPose.SetBoneContainer(&BoneContainer);
+		UAnimPoseExtensions::GetCompactPose(AnimPose, OutPose); 
 
-			FAnimationRuntime::MirrorPose(OutPose, MirrorDataTable->MirrorAxis, CompactPoseMirrorBones, ComponentSpaceRefRotations);
+		FAnimationRuntime::MirrorPose(OutPose, MirrorDataTable->MirrorAxis, CompactPoseMirrorBones, ComponentSpaceRefRotations);
 
-			for (int32 BoneIndex = 0; BoneIndex < BoneContainer.GetNumBones(); ++BoneIndex)
+		for (int32 BoneIndex = 0; BoneIndex < BoneContainer.GetNumBones(); ++BoneIndex)
+		{
+			FName SourceBoneName = BoneContainer.GetReferenceSkeleton().GetBoneName(BoneIndex);
+			FSkeletonPoseBoneIndex SkeletonBoneIndex(BoneIndex);
+			FCompactPoseBoneIndex CompactIndex = BoneContainer.GetCompactPoseIndexFromSkeletonPoseIndex(SkeletonBoneIndex);
+			FTransform BonePoseTargetLocal = OutPose[CompactIndex];
+			const FInt32Range KeyRangeToSet(AnimKey, AnimKey + 1);
+			if (Model->IsValidBoneTrackName(SourceBoneName))
 			{
-				FName SourceBoneName = BoneContainer.GetReferenceSkeleton().GetBoneName(BoneIndex);
-				FSkeletonPoseBoneIndex SkeletonBoneIndex(BoneIndex);
-				FCompactPoseBoneIndex CompactIndex = BoneContainer.GetCompactPoseIndexFromSkeletonPoseIndex(SkeletonBoneIndex);
-				FTransform BonePoseTargetLocal = OutPose[CompactIndex];
-				const FInt32Range KeyRangeToSet(AnimKey, AnimKey + 1);
-				if (Model->IsValidBoneTrackName(SourceBoneName))
-				{
-					Controller.UpdateBoneTrackKeys(SourceBoneName, KeyRangeToSet, { BonePoseTargetLocal.GetLocation() }, { BonePoseTargetLocal.GetRotation() }, { BonePoseTargetLocal.GetScale3D() });
-				}
-
-			}	
+				Controller.UpdateBoneTrackKeys(SourceBoneName, KeyRangeToSet, { BonePoseTargetLocal.GetLocation() }, { BonePoseTargetLocal.GetRotation() }, { BonePoseTargetLocal.GetScale3D() });
+			}
+
+		}	
 	}
 
 	if (bUpdateSyncMarkers)
