When MotionExtrantorModifyer is applied to skeletal mesh, it outputs warning about module loading of AnimationModifierLibrary at editor startup in ThirdPerson C++ project.
#fix Mark native animation modifier library to be loaded during PreDefault

diff --git a/Engine/Plugins/Animation/AnimationModifierLibrary/AnimationModifierLibrary.uplugin b/Engine/Plugins/Animation/AnimationModifierLibrary/AnimationModifierLibrary.uplugin
--- a/Engine/Plugins/Animation/AnimationModifierLibrary/AnimationModifierLibrary.uplugin
+++ b/Engine/Plugins/Animation/AnimationModifierLibrary/AnimationModifierLibrary.uplugin
@@ -20,7 +20,7 @@
 		{
 			"Name": "AnimationModifierLibrary",
 			"Type": "UncookedOnly",
-			"LoadingPhase": "Default"
+			"LoadingPhase": "PreDefault"
 		}
 	]
 }
