PR #13030: UCopyBonesModifier::OnApply_Implementation optimization
#fix reimplemented pull-request using SetBoneTrackKeys rather than UpdateBoneTrackKeys to minimize controller operations
#misc reserve key arrays when transforming between FVector and FVector3f in controller API

diff --git a/Engine/Plugins/Animation/AnimationData/Source/AnimationData/Private/AnimSequencerController.cpp b/Engine/Plugins/Animation/AnimationData/Source/AnimationData/Private/AnimSequencerController.cpp
--- a/Engine/Plugins/Animation/AnimationData/Source/AnimationData/Private/AnimSequencerController.cpp
+++ b/Engine/Plugins/Animation/AnimationData/Source/AnimationData/Private/AnimSequencerController.cpp
@@ -1639,6 +1639,10 @@ bool UAnimSequencerController::SetBoneTrackKeys(FName BoneName, const TArray<FVe
 	TArray<FQuat4f> FloatRotationalKeys;
 	TArray<FVector3f> FloatScalingKeys;
 
+	FloatPositionalKeys.Reserve(PositionalKeys.Num());
+	FloatRotationalKeys.Reserve(RotationalKeys.Num());
+	FloatScalingKeys.Reserve(ScalingKeys.Num());
+
 	Algo::Transform(PositionalKeys, FloatPositionalKeys, [](const FVector& V) { return FVector3f(V); });
 	Algo::Transform(RotationalKeys, FloatRotationalKeys, [](const FQuat& Q) { return FQuat4f(Q); });
 	Algo::Transform(ScalingKeys, FloatScalingKeys, [](const FVector& V) { return FVector3f(V); });
@@ -1767,6 +1771,9 @@ bool UAnimSequencerController::UpdateBoneTrackKeys(FName BoneName, const FInt32R
 	TArray<FVector3f> FloatPositionalKeys;
 	TArray<FQuat4f> FloatRotationalKeys;
 	TArray<FVector3f> FloatScalingKeys;
+	FloatPositionalKeys.Reserve(PositionalKeys.Num());
+	FloatRotationalKeys.Reserve(RotationalKeys.Num());
+	FloatScalingKeys.Reserve(ScalingKeys.Num());
 
 	Algo::Transform(PositionalKeys, FloatPositionalKeys, [](const FVector& V) { return FVector3f(V); });
 	Algo::Transform(RotationalKeys, FloatRotationalKeys, [](const FQuat& Q) { return FQuat4f(Q); });
diff --git a/Engine/Plugins/Animation/AnimationModifierLibrary/Source/AnimationModifierLibrary/Private/CopyBonesModifier.cpp b/Engine/Plugins/Animation/AnimationModifierLibrary/Source/AnimationModifierLibrary/Private/CopyBonesModifier.cpp
--- a/Engine/Plugins/Animation/AnimationModifierLibrary/Source/AnimationModifierLibrary/Private/CopyBonesModifier.cpp
+++ b/Engine/Plugins/Animation/AnimationModifierLibrary/Source/AnimationModifierLibrary/Private/CopyBonesModifier.cpp
@@ -77,27 +77,61 @@ void UCopyBonesModifier::OnApply_Implementation(UAnimSequence* Animation)
 
 	// Get the transform of all the source bones in the desired space
 	const int32 NumKeys = Model->GetNumberOfKeys();
-	FScopedSlowTask SlowTask(static_cast<float>(NumKeys));
-	SlowTask.MakeDialog();
+	const FInt32Range KeyRangeToSet(0, NumKeys);
+	struct FPerBoneData
+	{
+		TArray<FVector> PositionalKeys;
+		TArray<FQuat> RotationalKeys;
+		TArray<FVector> ScalingKeys;
+	};
+
+	TArray<FPerBoneData> PerBoneDataContainer;
+
+	// Pre-allocate bone arrays for each copy
+	const int32 NumBoneCopies = CopyBoneDataContainer.Num();
+	PerBoneDataContainer.SetNum(CopyBoneDataContainer.Num());
+	for (int32 Index = 0; Index < NumBoneCopies; ++Index)
+	{
+		PerBoneDataContainer[Index].PositionalKeys.Reserve(NumKeys);
+		PerBoneDataContainer[Index].RotationalKeys.Reserve(NumKeys);
+		PerBoneDataContainer[Index].ScalingKeys.Reserve(NumKeys);
+	}
+
+	FScopedSlowTask SlowTask(static_cast<float>(CopyBoneDataContainer.Num() + NumKeys));
 	for (int32 AnimKey = 0; AnimKey < NumKeys; AnimKey++)
 	{
+		FAnimPose AnimPose;
+		UAnimPoseExtensions::GetAnimPoseAtFrame(Animation, AnimKey, FAnimPoseEvaluationOptions(), AnimPose);
+
 		SlowTask.EnterProgressFrame();
-		for (FCopyBoneData& Data : CopyBoneDataContainer)
+		for (int32 Index = 0; Index < NumBoneCopies; ++Index)
 		{
-			FAnimPose AnimPose;
-			UAnimPoseExtensions::GetAnimPoseAtFrame(Animation, AnimKey, FAnimPoseEvaluationOptions(), AnimPose);
-
-			FTransform BonePose = UAnimPoseExtensions::GetBonePose(AnimPose, Data.SourceBoneName, BonePoseSpace);
-
+			FPerBoneData& PerBoneData = PerBoneDataContainer[Index];
+			FCopyBoneData& Data = CopyBoneDataContainer[Index];
+			
+			// Make a copy of the pose to deal with potential parent-chain issues
+			FAnimPose AnimPoseCopy = AnimPose;
+			FTransform BonePose = UAnimPoseExtensions::GetBonePose(AnimPoseCopy, Data.SourceBoneName, BonePoseSpace);
+			
 			// UAnimDataController::UpdateBoneTrackKeys expects local transforms so we need to convert the source transforms to target bone local transforms first. 
-			UAnimPoseExtensions::SetBonePose(AnimPose, BonePose, Data.TargetBoneName, BonePoseSpace);
-			FTransform BonePoseTargetLocal = UAnimPoseExtensions::GetBonePose(AnimPose, Data.TargetBoneName, EAnimPoseSpaces::Local);
+			UAnimPoseExtensions::SetBonePose(AnimPoseCopy, BonePose, Data.TargetBoneName, BonePoseSpace);
+			FTransform BonePoseTargetLocal = UAnimPoseExtensions::GetBonePose(AnimPoseCopy, Data.TargetBoneName, EAnimPoseSpaces::Local);
 
-			const FInt32Range KeyRangeToSet(AnimKey, AnimKey + 1);
-			Controller.UpdateBoneTrackKeys(Data.TargetBoneName, KeyRangeToSet, { BonePoseTargetLocal.GetLocation() }, { BonePoseTargetLocal.GetRotation() }, { BonePoseTargetLocal.GetScale3D() }, bShouldTransact);
+			PerBoneData.PositionalKeys.Add(BonePoseTargetLocal.GetLocation());
+			PerBoneData.RotationalKeys.Add(BonePoseTargetLocal.GetRotation());
+			PerBoneData.ScalingKeys.Add(BonePoseTargetLocal.GetScale3D());
 		}
 	}
 
+	
+	for (int32 Index = 0; Index < NumBoneCopies; ++Index)
+	{
+		SlowTask.EnterProgressFrame();
+		FCopyBoneData& Data = CopyBoneDataContainer[Index];
+		FPerBoneData& PerBoneData = PerBoneDataContainer[Index];
+		Controller.SetBoneTrackKeys(Data.TargetBoneName, PerBoneData.PositionalKeys, PerBoneData.RotationalKeys, PerBoneData.ScalingKeys, bShouldTransact);
+	}
+
 	// Done editing animation data
 	Controller.CloseBracket(bShouldTransact);
 }
