PR #13938: Fix CopyBonesModifier for bones without pre-existing curves

diff --git a/Engine/Plugins/Animation/AnimationModifierLibrary/Source/AnimationModifierLibrary/Private/CopyBonesModifier.cpp b/Engine/Plugins/Animation/AnimationModifierLibrary/Source/AnimationModifierLibrary/Private/CopyBonesModifier.cpp
--- a/Engine/Plugins/Animation/AnimationModifierLibrary/Source/AnimationModifierLibrary/Private/CopyBonesModifier.cpp
+++ b/Engine/Plugins/Animation/AnimationModifierLibrary/Source/AnimationModifierLibrary/Private/CopyBonesModifier.cpp
@@ -131,6 +131,13 @@ void UCopyBonesModifier::OnApply_Implementation(UAnimSequence* Animation)
 		SlowTask.EnterProgressFrame();
 		FCopyBoneData& Data = CopyBoneDataContainer[Index];
 		FPerBoneData& PerBoneData = PerBoneDataContainer[Index];
+
+		// If there are no pre-existing curves for the target bone we have to add them before being able to set the keys
+		if(!Model->IsValidBoneTrackName(Data.TargetBoneName))
+		{
+			Controller.AddBoneCurve(Data.TargetBoneName, bShouldTransact);
+		}
+		
 		Controller.SetBoneTrackKeys(Data.TargetBoneName, PerBoneData.PositionalKeys, PerBoneData.RotationalKeys, PerBoneData.ScalingKeys, bShouldTransact);
 	}
 
