AnimationModifierLibrary - adding UCurveFromSyncMarkersModifier to generate a curve from synch markers

#[review]-52715038 [REDACTED_HANDLE]
#Synced-CL 52711943

#Synced-CL 52765189

diff --git a/Engine/Plugins/Animation/AnimationModifierLibrary/Source/AnimationModifierLibrary/Private/CurveFromSyncMarkersModifier.cpp b/Engine/Plugins/Animation/AnimationModifierLibrary/Source/AnimationModifierLibrary/Private/CurveFromSyncMarkersModifier.cpp
new file mode 100644
--- /dev/null
+++ b/Engine/Plugins/Animation/AnimationModifierLibrary/Source/AnimationModifierLibrary/Private/CurveFromSyncMarkersModifier.cpp
@@ -0,0 +1,121 @@
+// Copyright Epic Games, Inc. All Rights Reserved.
+
+#include "CurveFromSyncMarkersModifier.h"
+#include "Animation/AnimSequence.h"
+#include "AnimationBlueprintLibrary.h"
+#include "EngineLogs.h"
+
+#include UE_INLINE_GENERATED_CPP_BY_NAME(CurveFromSyncMarkersModifier)
+
+#define LOCTEXT_NAMESPACE "CurveFromSyncMarkersModifier"
+
+void UCurveFromSyncMarkersModifier::OnApply_Implementation(UAnimSequence* Animation)
+{
+	if (Animation == nullptr)
+	{
+		UE_LOGF(LogAnimation, Error, "CurveFromSyncMarkersModifier failed. Reason: Invalid Animation");
+		return;
+	}
+
+	if (CurveName.IsNone())
+	{
+		UE_LOGF(LogAnimation, Error, "CurveFromSyncMarkersModifier failed. Reason: CurveName is None");
+		return;
+	}
+
+	if (SyncMarkerValues.IsEmpty())
+	{
+		UE_LOGF(LogAnimation, Warning, "CurveFromSyncMarkersModifier: SyncMarkerValues is empty, no curve will be generated");
+		return;
+	}
+
+	// Build a lookup from sync marker name to float value
+	TMap<FName, float> MarkerValueMap;
+	for (const FSyncMarkerCurveEntry& Entry : SyncMarkerValues)
+	{
+		if (!Entry.SyncMarkerName.IsNone())
+		{
+			MarkerValueMap.Add(Entry.SyncMarkerName, Entry.Value);
+		}
+	}
+
+	// Collect curve keys from authored sync markers whose names match the input list
+	TArray<FRichCurveKey> CurveKeys;
+	for (const FAnimSyncMarker& Marker : Animation->AuthoredSyncMarkers)
+	{
+		if (const float* ValuePtr = MarkerValueMap.Find(Marker.MarkerName))
+		{
+			FRichCurveKey& Key = CurveKeys.AddDefaulted_GetRef();
+			Key.Time = Marker.Time;
+			Key.Value = *ValuePtr;
+		}
+	}
+
+	if (CurveKeys.IsEmpty())
+	{
+		UE_LOGF(LogAnimation, Warning, "CurveFromSyncMarkersModifier: No matching sync markers found. Curve will not be generated");
+		return;
+	}
+
+	// AuthoredSyncMarkers should already be in time order, but sort defensively
+	CurveKeys.Sort([](const FRichCurveKey& A, const FRichCurveKey& B) { return A.Time < B.Time; });
+
+	// For looping animations, add boundary keys at time 0 and animation end so the curve
+	// is continuous across the wrap seam. The boundary value is linearly interpolated across
+	// the segment from the last key forward (past the animation end) to the first key.
+	if (Animation->bLoop)
+	{
+		const float AnimLength = Animation->GetPlayLength();
+		const FRichCurveKey& FirstKey = CurveKeys[0];
+		const FRichCurveKey& LastKey  = CurveKeys.Last();
+
+		// Segment spans from LastKey.Time to FirstKey.Time + AnimLength (wrapping around)
+		const float SegmentDuration = (FirstKey.Time + AnimLength) - LastKey.Time;
+		const float BoundaryValue = (SegmentDuration > UE_KINDA_SMALL_NUMBER)
+			? FMath::Lerp(LastKey.Value, FirstKey.Value, (AnimLength - LastKey.Time) / SegmentDuration)
+			: LastKey.Value;
+
+		// Prepend at time 0 only when no marker already sits there
+		if (FirstKey.Time > UE_KINDA_SMALL_NUMBER)
+		{
+			FRichCurveKey StartKey;
+			StartKey.Time  = 0.f;
+			StartKey.Value = BoundaryValue;
+			CurveKeys.Insert(StartKey, 0);
+		}
+
+		// Append at animation end only when no marker already sits there
+		// (using CurveKeys.Last() instead of LastKey since CurveKeys could have been reallocated)
+		if ((AnimLength - CurveKeys.Last().Time) > UE_KINDA_SMALL_NUMBER)
+		{
+			FRichCurveKey EndKey;
+			EndKey.Time  = AnimLength;
+			EndKey.Value = BoundaryValue;
+			CurveKeys.Add(EndKey);
+		}
+	}
+
+	IAnimationDataController& Controller = Animation->GetController();
+	Controller.OpenBracket(LOCTEXT("CurveFromSyncMarkers_Bracket", "Adding Curve From Sync Markers"));
+
+	const FAnimationCurveIdentifier CurveId = UAnimationCurveIdentifierExtensions::GetCurveIdentifier(
+		Animation->GetSkeleton(), CurveName, ERawCurveTrackTypes::RCT_Float);
+
+	if (Controller.AddCurve(CurveId))
+	{
+		Controller.SetCurveKeys(CurveId, CurveKeys);
+	}
+
+	Controller.CloseBracket();
+}
+
+void UCurveFromSyncMarkersModifier::OnRevert_Implementation(UAnimSequence* Animation)
+{
+	if (bRemoveCurveOnRevert)
+	{
+		const bool bRemoveNameFromSkeleton = false;
+		UAnimationBlueprintLibrary::RemoveCurve(Animation, CurveName, bRemoveNameFromSkeleton);
+	}
+}
+
+#undef LOCTEXT_NAMESPACE
diff --git a/Engine/Plugins/Animation/AnimationModifierLibrary/Source/AnimationModifierLibrary/Public/CurveFromSyncMarkersModifier.h b/Engine/Plugins/Animation/AnimationModifierLibrary/Source/AnimationModifierLibrary/Public/CurveFromSyncMarkersModifier.h
new file mode 100644
--- /dev/null
+++ b/Engine/Plugins/Animation/AnimationModifierLibrary/Source/AnimationModifierLibrary/Public/CurveFromSyncMarkersModifier.h
@@ -0,0 +1,44 @@
+// Copyright Epic Games, Inc. All Rights Reserved.
+
+#pragma once
+
+#include "AnimationModifier.h"
+#include "CurveFromSyncMarkersModifier.generated.h"
+
+/** A sync marker name and the float value to assign to the curve at that marker's time */
+USTRUCT(Experimental, BlueprintType)
+struct FSyncMarkerCurveEntry
+{
+	GENERATED_BODY()
+
+	UPROPERTY(EditAnywhere, BlueprintReadWrite, Category = "Settings")
+	FName SyncMarkerName;
+
+	UPROPERTY(EditAnywhere, BlueprintReadWrite, Category = "Settings")
+	float Value = 0.f;
+};
+
+/** Generates a float curve from the authored sync markers on an animation sequence.
+ *  Each sync marker listed in SyncMarkerValues contributes a curve key at the marker's time
+ *  with the corresponding float value. */
+UCLASS(Experimental)
+class UCurveFromSyncMarkersModifier : public UAnimationModifier
+{
+	GENERATED_BODY()
+
+public:
+	/** List of sync marker names and the float value each should contribute to the generated curve */
+	UPROPERTY(EditAnywhere, BlueprintReadWrite, Category = "Settings")
+	TArray<FSyncMarkerCurveEntry> SyncMarkerValues;
+
+	/** Name of the float curve to add to the animation */
+	UPROPERTY(EditAnywhere, BlueprintReadWrite, Category = "Settings")
+	FName CurveName;
+
+	/** If true, the generated curve is removed when the modifier is reverted */
+	UPROPERTY(EditAnywhere, BlueprintReadWrite, Category = "Settings")
+	bool bRemoveCurveOnRevert = true;
+
+	virtual void OnApply_Implementation(UAnimSequence* Animation) override;
+	virtual void OnRevert_Implementation(UAnimSequence* Animation) override;
+};
